Plugin that adds EdMode based on InteractiveToolsFramework which contains sample Tool implementations (drop empty Actors on click, draw polycurve on selected mesh)
#rnx

diff --git a/Engine/Plugins/Experimental/SampleToolsEditorMode/SampleToolsEditorMode.uplugin b/Engine/Plugins/Experimental/SampleToolsEditorMode/SampleToolsEditorMode.uplugin
new file mode 100644
--- /dev/null
+++ b/Engine/Plugins/Experimental/SampleToolsEditorMode/SampleToolsEditorMode.uplugin
@@ -0,0 +1,24 @@
+{
+	"FileVersion": 3,
+	"Version": 1,
+	"VersionName": "0.1",
+	"FriendlyName": "Sample Tools Mode",
+	"Description": "Adds Sample Tools Editor Mode Tab",
+	"Category": "Other",
+	"CreatedBy": "Epic Games, Inc.",
+	"CreatedByURL": "http://epicgames.com",
+	"DocsURL": "",
+	"MarketplaceURL": "",
+	"SupportURL": "",
+	"CanContainContent": false,
+	"IsBetaVersion": true,
+	"Installed": false,
+	"Modules": [
+		{
+			"Name": "SampleToolsEditorMode",
+			"Type": "Editor",
+			"LoadingPhase": "Default"
+		}
+	],
+	"Plugins": []
+}
\ No newline at end of file
diff --git a/Engine/Plugins/Experimental/SampleToolsEditorMode/Source/Private/SampleTools/CreateActorSampleTool.cpp b/Engine/Plugins/Experimental/SampleToolsEditorMode/Source/Private/SampleTools/CreateActorSampleTool.cpp
new file mode 100644
--- /dev/null
+++ b/Engine/Plugins/Experimental/SampleToolsEditorMode/Source/Private/SampleTools/CreateActorSampleTool.cpp
@@ -0,0 +1,106 @@
+// Copyright 1998-2019 Epic Games, Inc. All Rights Reserved.
+
+#include "CreateActorSampleTool.h"
+#include "ToolBuilderUtil.h"
+
+// localization namespace
+#define LOCTEXT_NAMESPACE "UCreateActorSampleTool"
+
+/*
+ * ToolBuilder
+ */
+
+
+bool UCreateActorSampleToolBuilder::CanBuildTool(const FToolBuilderState & SceneState) const
+{
+	return (this->AssetAPI != nullptr);
+}
+
+UInteractiveTool* UCreateActorSampleToolBuilder::BuildTool(const FToolBuilderState & SceneState) const
+{
+	UCreateActorSampleTool* NewTool = NewObject<UCreateActorSampleTool>();
+	NewTool->SetWorld(SceneState.World);
+	NewTool->SetAssetAPI(AssetAPI);
+	return NewTool;
+}
+
+
+
+/*
+ * Tool
+ */
+
+
+UCreateActorSampleTool::UCreateActorSampleTool()
+{
+	PlaceOnObjects = true;
+	GroundHeight = 0.0f;
+}
+
+
+void UCreateActorSampleTool::SetWorld(UWorld* World)
+{
+	this->TargetWorld = World;
+}
+
+void UCreateActorSampleTool::SetAssetAPI(IToolsContextAssetAPI* AssetAPIIn)
+{
+	this->AssetAPI = AssetAPIIn;
+}
+
+
+void UCreateActorSampleTool::Setup()
+{
+	USingleClickTool::Setup();
+
+	// provide self as property object
+	ToolPropertyObjects.Add(this);
+}
+
+
+
+
+void UCreateActorSampleTool::OnClicked(const FInputDeviceRay& ClickPos)
+{
+	// we will create actor at this position
+	FVector NewActorPos = FVector::ZeroVector;
+
+	// cast ray into world to find hit position
+	FVector RayStart = ClickPos.WorldRay.Origin;
+	FVector RayEnd = ClickPos.WorldRay.PointAt(999999);
+	FCollisionObjectQueryParams QueryParams(FCollisionObjectQueryParams::AllObjects);
+	FHitResult Result;
+	bool bHitWorld = TargetWorld->LineTraceSingleByObjectType(Result, RayStart, RayEnd, QueryParams);
+	if (PlaceOnObjects && bHitWorld)
+	{
+		NewActorPos = Result.ImpactPoint;
+	}
+	else
+	{
+		// hit the ground plane
+		FPlane GroundPlane(FVector(0,0,GroundHeight) , FVector(0, 0, 1));
+		NewActorPos = FMath::RayPlaneIntersection(ClickPos.WorldRay.Origin, ClickPos.WorldRay.Direction, GroundPlane);
+	}
+
+	// create new actor
+	FRotator Rotation(0.0f, 0.0f, 0.0f);
+	FActorSpawnParameters SpawnInfo;
+	AActor* NewActor = TargetWorld->SpawnActor<AActor>(FVector::ZeroVector, Rotation, SpawnInfo);
+
+	// create root component
+	USceneComponent* RootComponent = NewObject<USceneComponent>(NewActor, USceneComponent::GetDefaultSceneRootVariableName(), RF_Transactional);
+	RootComponent->Mobility = EComponentMobility::Movable;
+	RootComponent->bVisualizeComponent = true;
+	RootComponent->SetWorldTransform(FTransform(NewActorPos));
+
+	NewActor->SetRootComponent(RootComponent);
+	NewActor->AddInstanceComponent(RootComponent);
+	RootComponent->RegisterComponent();
+
+}
+
+
+
+
+
+#undef LOCTEXT_NAMESPACE
\ No newline at end of file
diff --git a/Engine/Plugins/Experimental/SampleToolsEditorMode/Source/Private/SampleTools/CreateActorSampleTool.h b/Engine/Plugins/Experimental/SampleToolsEditorMode/Source/Private/SampleTools/CreateActorSampleTool.h
new file mode 100644
--- /dev/null
+++ b/Engine/Plugins/Experimental/SampleToolsEditorMode/Source/Private/SampleTools/CreateActorSampleTool.h
@@ -0,0 +1,71 @@
+// Copyright 1998-2019 Epic Games, Inc. All Rights Reserved.
+
+#pragma once
+
+#include "CoreMinimal.h"
+#include "UObject/NoExportTypes.h"
+#include "InteractiveToolBuilder.h"
+#include "BaseTools/SingleClickTool.h"
+#include "CreateActorSampleTool.generated.h"
+
+
+/**
+ * Builder for UCreateActorSampleTool
+ */
+UCLASS()
+class SAMPLETOOLSEDITORMODE_API UCreateActorSampleToolBuilder : public UInteractiveToolBuilder
+{
+	GENERATED_BODY()
+
+public:
+	IToolsContextAssetAPI* AssetAPI;
+
+	UCreateActorSampleToolBuilder() 
+	{
+		AssetAPI = nullptr;
+	}
+
+	virtual bool CanBuildTool(const FToolBuilderState& SceneState) const override;
+	virtual UInteractiveTool* BuildTool(const FToolBuilderState& SceneState) const override;
+};
+
+
+
+/**
+ * UCreateActorSampleTool is an example Tool that drops an empty Actor at each position the user 
+ * clicks left mouse button. The Actors are placed at the first ray intersection in the scene,
+ * or on a ground plane if no scene objects are hit. All the action is in the ::OnClicked handler.
+ */
+UCLASS()
+class SAMPLETOOLSEDITORMODE_API UCreateActorSampleTool : public USingleClickTool
+{
+	GENERATED_BODY()
+
+public:
+	UCreateActorSampleTool();
+
+	virtual void SetWorld(UWorld* World);
+	virtual void SetAssetAPI(IToolsContextAssetAPI* AssetAPI);
+
+	virtual void Setup() override;
+
+	virtual void OnClicked(const FInputDeviceRay& ClickPos);
+
+
+protected:
+	/** Place actors on existing objects */
+	UPROPERTY(EditAnywhere, Category = Options, meta = (DisplayName = "Place On Objects"))
+	bool PlaceOnObjects;
+
+	/** Height of ground plane */
+	UPROPERTY(EditAnywhere, Category = Options, meta = (DisplayName = "Ground Height", UIMin = "-1000.0", UIMax = "1000.0", ClampMin = "-1000000", ClampMax = "1000000.0"))
+	float GroundHeight;
+
+
+protected:
+	/** target World we will raycast into and create Actor in */
+	UWorld* TargetWorld;
+
+	/** Access to the ToolContext's Asset Creation API. This is not currently used, but can be used to (eg) add Components, etc*/
+	IToolsContextAssetAPI* AssetAPI;
+};
\ No newline at end of file
diff --git a/Engine/Plugins/Experimental/SampleToolsEditorMode/Source/Private/SampleTools/DrawCurveOnMeshSampleTool.cpp b/Engine/Plugins/Experimental/SampleToolsEditorMode/Source/Private/SampleTools/DrawCurveOnMeshSampleTool.cpp
new file mode 100644
--- /dev/null
+++ b/Engine/Plugins/Experimental/SampleToolsEditorMode/Source/Private/SampleTools/DrawCurveOnMeshSampleTool.cpp
@@ -0,0 +1,101 @@
+// Copyright 1998-2019 Epic Games, Inc. All Rights Reserved.
+
+#include "DrawCurveOnMeshSampleTool.h"
+#include "ToolBuilderUtil.h"
+
+// localization namespace
+#define LOCTEXT_NAMESPACE "UDrawCurveOnMeshSampleTool"
+
+/*
+ * ToolBuilder
+ */
+
+
+UMeshSurfacePointTool* UDrawCurveOnMeshSampleToolBuilder::CreateNewTool(const FToolBuilderState & SceneState) const
+{
+	UDrawCurveOnMeshSampleTool* NewTool = NewObject<UDrawCurveOnMeshSampleTool>();
+	return NewTool;
+}
+
+
+
+/*
+ * Tool
+ */
+
+
+UDrawCurveOnMeshSampleTool::UDrawCurveOnMeshSampleTool()
+{
+	Thickness = 4.0f;
+	DepthBias = 0.0f;
+	MinSpacing = 1.0;
+	NormalOffset = 0.25;
+	Color = FLinearColor(255, 0, 0);
+	bScreenSpace = true;
+}
+
+
+void UDrawCurveOnMeshSampleTool::Setup()
+{
+	UMeshSurfacePointTool::Setup();
+
+	// provide self as property object
+	ToolPropertyObjects.Add(this);
+}
+
+
+
+void UDrawCurveOnMeshSampleTool::Render(IToolsContextRenderAPI* RenderAPI)
+{
+	auto PDI = RenderAPI->GetPrimitiveDrawInterface();
+	
+	int NumPts = Positions.Num();
+	for (int i = 0; i < NumPts - 1; ++i)
+	{
+		FVector A = Positions[i] + NormalOffset*Normals[i];
+		FVector B = Positions[i+1] + NormalOffset*Normals[i+1];
+		PDI->DrawLine(A, B, Color, 0, Thickness, DepthBias, bScreenSpace);
+	}
+}
+
+
+
+
+void UDrawCurveOnMeshSampleTool::OnBeginDrag(const FRay& Ray)
+{
+	Positions.Reset();
+	Normals.Reset();
+	
+	FHitResult OutHit;
+	if (HitTest(Ray, OutHit))
+	{
+		Positions.Add(OutHit.ImpactPoint);
+		Normals.Add(OutHit.ImpactNormal);
+	}
+
+}
+
+void UDrawCurveOnMeshSampleTool::OnUpdateDrag(const FRay& Ray)
+{
+	FHitResult OutHit;
+	if (HitTest(Ray, OutHit))
+	{
+		if ( FVector::Dist(OutHit.ImpactPoint, Positions[Positions.Num()-1]) > MinSpacing)
+		{
+			Positions.Add(OutHit.ImpactPoint);
+			Normals.Add(OutHit.ImpactNormal);
+		}
+	}
+}
+
+void UDrawCurveOnMeshSampleTool::OnEndDrag(const FRay& Ray)
+{
+
+}
+
+
+
+
+
+
+#undef LOCTEXT_NAMESPACE
\ No newline at end of file
diff --git a/Engine/Plugins/Experimental/SampleToolsEditorMode/Source/Private/SampleTools/DrawCurveOnMeshSampleTool.h b/Engine/Plugins/Experimental/SampleToolsEditorMode/Source/Private/SampleTools/DrawCurveOnMeshSampleTool.h
new file mode 100644
--- /dev/null
+++ b/Engine/Plugins/Experimental/SampleToolsEditorMode/Source/Private/SampleTools/DrawCurveOnMeshSampleTool.h
@@ -0,0 +1,75 @@
+// Copyright 1998-2019 Epic Games, Inc. All Rights Reserved.
+
+#pragma once
+
+#include "CoreMinimal.h"
+#include "UObject/NoExportTypes.h"
+#include "InteractiveToolBuilder.h"
+#include "BaseTools/MeshSurfacePointTool.h"
+#include "DrawCurveOnMeshSampleTool.generated.h"
+
+
+/**
+ * UMeshSurfacePointToolBuilder override for UDrawCurveOnMeshSampleTool
+ */
+UCLASS()
+class SAMPLETOOLSEDITORMODE_API UDrawCurveOnMeshSampleToolBuilder : public UMeshSurfacePointToolBuilder
+{
+	GENERATED_BODY()
+
+public:
+	virtual UMeshSurfacePointTool* CreateNewTool(const FToolBuilderState& SceneState) const override;
+};
+
+
+
+/**
+ * UDrawCurveOnMeshSampleTool is a sample Tool that allows the user to draw curves on the surface of
+ * a selected Mesh Component. The various rendering properties of the polycurve are exposed and can be tweaked.
+ * Nothing is done with the curve, it is just drawn by ::Render() and discarded when the Tool exits.
+ */
+UCLASS()
+class SAMPLETOOLSEDITORMODE_API UDrawCurveOnMeshSampleTool : public UMeshSurfacePointTool
+{
+	GENERATED_BODY()
+
+public:
+	UDrawCurveOnMeshSampleTool();
+
+	// UInteractiveTool API
+
+	virtual void Setup() override;
+
+	virtual void Render(IToolsContextRenderAPI* RenderAPI) override;
+
+	// UMeshSurfacePointTool API
+	virtual void OnBeginDrag(const FRay& Ray) override;
+	virtual void OnUpdateDrag(const FRay& Ray) override;
+	virtual void OnEndDrag(const FRay& Ray) override;
+
+protected:
+	UPROPERTY(EditAnywhere, Category = Options)
+	FLinearColor Color;
+
+	UPROPERTY(EditAnywhere, Category = Options, meta = (DisplayName = "Thickness", UIMin = "0.25", UIMax = "10.0", ClampMin = "0.01", ClampMax = "1000.0"))
+	float Thickness;
+
+	UPROPERTY(EditAnywhere, Category = Options, meta = (DisplayName = "Min Spacing", UIMin = "0.01", UIMax = "10.0"))
+	float MinSpacing;
+
+	UPROPERTY(EditAnywhere, Category = Options, meta = (DisplayName = "Offset", UIMin = "0.0", UIMax = "10.0", ClampMin = "-1000.0", ClampMax = "1000.0"))
+	float NormalOffset;
+
+	UPROPERTY(EditAnywhere, Category = Options, meta = (DisplayName = "Depth Bias", UIMin = "-10.0", UIMax = "10.0"))
+	float DepthBias;
+
+	UPROPERTY(EditAnywhere, Category = Options)
+	bool bScreenSpace;
+
+	UPROPERTY()
+	TArray<FVector> Positions;
+
+	UPROPERTY()
+	TArray<FVector> Normals;
+
+};
\ No newline at end of file
diff --git a/Engine/Plugins/Experimental/SampleToolsEditorMode/Source/Private/SampleToolsEditorMode.cpp b/Engine/Plugins/Experimental/SampleToolsEditorMode/Source/Private/SampleToolsEditorMode.cpp
new file mode 100644
--- /dev/null
+++ b/Engine/Plugins/Experimental/SampleToolsEditorMode/Source/Private/SampleToolsEditorMode.cpp
@@ -0,0 +1,233 @@
+// Copyright 1998-2019 Epic Games, Inc. All Rights Reserved.
+
+#include "SampleToolsEditorMode.h"
+#include "SampleToolsEditorModeToolkit.h"
+#include "Toolkits/ToolkitManager.h"
+#include "EditorViewportClient.h"
+#include "EditorModeManager.h"
+
+
+//////////////////////////////////////////////////////////////////////////
+////////////////////////////////////////////////////////////////////////// 
+// AddYourTool Step 1 - include the header file for your Tool here
+//////////////////////////////////////////////////////////////////////////
+////////////////////////////////////////////////////////////////////////// 
+#include "BaseTools/SingleClickTool.h"
+#include "BaseTools/MeshSurfacePointTool.h"
+#include "SampleTools/CreateActorSampleTool.h"
+#include "SampleTools/DrawCurveOnMeshSampleTool.h"
+
+// step 2: register a ToolBuilder in FSampleToolsEditorMode::Enter()
+// step 3: add a button in FSampleToolsEditorModeToolkit::Init()
+
+
+#define LOCTEXT_NAMESPACE "FSampleToolsEditorMode"
+
+const FEditorModeID FSampleToolsEditorMode::EM_SampleToolsEditorModeId = TEXT("EM_SampleToolsEditorMode");
+
+
+FSampleToolsEditorMode::FSampleToolsEditorMode()
+{
+	ToolsContext = nullptr;
+}
+
+
+FSampleToolsEditorMode::~FSampleToolsEditorMode()
+{
+	// this should have happend already in ::Exit()
+	if (ToolsContext != nullptr)
+	{
+		ToolsContext->ShutdownContext();
+		ToolsContext = nullptr;
+	}
+}
+
+
+void FSampleToolsEditorMode::ActorSelectionChangeNotify()
+{
+	// @todo support selection change
+}
+
+
+
+void FSampleToolsEditorMode::Tick(FEditorViewportClient* ViewportClient, float DeltaTime)
+{
+	FEdMode::Tick(ViewportClient, DeltaTime);
+
+	// give ToolsContext a chance to tick
+	if (ToolsContext != nullptr)
+	{
+		ToolsContext->Tick(ViewportClient, DeltaTime);
+	}
+}
+
+
+
+
+void FSampleToolsEditorMode::Render(const FSceneView* View, FViewport* Viewport, FPrimitiveDrawInterface* PDI)
+{
+	FEdMode::Render(View, Viewport, PDI);
+
+	// give ToolsContext a chance to render
+	if (ToolsContext != nullptr)
+	{
+		ToolsContext->Render(View, Viewport, PDI);
+	}
+}
+
+
+
+
+
+//
+// Input device event tracking. We forward input events to the ToolsContext adapter for handling.
+// 
+
+bool FSampleToolsEditorMode::InputKey(FEditorViewportClient* ViewportClient, FViewport* Viewport, FKey Key, EInputEvent Event)
+{
+	bool bHandled = FEdMode::InputKey(ViewportClient, Viewport, Key, Event);
+	bHandled |= ToolsContext->InputKey(ViewportClient, Viewport, Key, Event);
+	return bHandled;
+}
+
+bool FSampleToolsEditorMode::InputAxis(FEditorViewportClient* InViewportClient, FViewport* Viewport, int32 ControllerId, FKey Key, float Delta, float DeltaTime)
+{
+	// mouse axes: EKeys::MouseX, EKeys::MouseY, EKeys::MouseWheelAxis
+	return FEdMode::InputAxis(InViewportClient, Viewport, ControllerId, Key, Delta, DeltaTime);
+}
+
+
+bool FSampleToolsEditorMode::StartTracking(FEditorViewportClient* InViewportClient, FViewport* InViewport) 
+{ 
+	bool bHandled = FEdMode::StartTracking(InViewportClient, InViewport);
+	bHandled |= ToolsContext->StartTracking(InViewportClient, InViewport);
+	return bHandled;
+}
+
+bool FSampleToolsEditorMode::CapturedMouseMove(FEditorViewportClient* InViewportClient, FViewport* InViewport, int32 InMouseX, int32 InMouseY)
+{
+	bool bHandled = ToolsContext->CapturedMouseMove(InViewportClient, InViewport, InMouseX, InMouseY);
+	return bHandled;
+}
+
+bool FSampleToolsEditorMode::EndTracking(FEditorViewportClient* InViewportClient, FViewport* InViewport)
+{
+	bool bHandled = ToolsContext->EndTracking(InViewportClient, InViewport);
+	return bHandled;
+}
+
+
+
+bool FSampleToolsEditorMode::MouseEnter(FEditorViewportClient* ViewportClient, FViewport* Viewport, int32 x, int32 y)
+{
+	bool bHandled = ToolsContext->MouseEnter(ViewportClient, Viewport, x, y);
+	return bHandled;
+}
+
+bool FSampleToolsEditorMode::MouseMove(FEditorViewportClient* ViewportClient, FViewport* Viewport, int32 x, int32 y)
+{
+	bool bHandled = ToolsContext->MouseMove(ViewportClient, Viewport, x, y);
+	return bHandled;
+}
+
+bool FSampleToolsEditorMode::MouseLeave(FEditorViewportClient* ViewportClient, FViewport* Viewport)
+{
+	bool bHandled = ToolsContext->MouseLeave(ViewportClient, Viewport);
+	return bHandled;
+}
+
+
+
+bool FSampleToolsEditorMode::ReceivedFocus(FEditorViewportClient* ViewportClient, FViewport* Viewport)
+{
+	return FEdMode::ReceivedFocus(ViewportClient, Viewport);
+}
+
+
+bool FSampleToolsEditorMode::LostFocus(FEditorViewportClient* ViewportClient, FViewport* Viewport)
+{
+	return FEdMode::LostFocus(ViewportClient, Viewport);
+}
+
+
+
+
+
+
+void FSampleToolsEditorMode::Enter()
+{
+	FEdMode::Enter();
+
+	if (!Toolkit.IsValid() && UsesToolkits())
+	{
+		Toolkit = MakeShareable(new FSampleToolsEditorModeToolkit);
+		Toolkit->Init(Owner->GetToolkitHost());
+	}
+
+	// initialize the adapter that attaches the ToolsContext to this FEdMode
+	ToolsContext = NewObject<UEdModeInteractiveToolsContext>(GetTransientPackage(), TEXT("ToolsContext"), RF_Transient);
+	ToolsContext->InitializeContextFromEdMode(this);
+
+
+	//////////////////////////////////////////////////////////////////////////
+	//////////////////////////////////////////////////////////////////////////
+	// AddYourTool Step 2 - register a ToolBuilder for your Tool here.
+	// The string name you pass to the ToolManager is used to select/activate your ToolBuilder later.
+	//////////////////////////////////////////////////////////////////////////
+	////////////////////////////////////////////////////////////////////////// 
+
+
+	auto CreateActorSampleToolBuilder = NewObject<UCreateActorSampleToolBuilder>();
+	CreateActorSampleToolBuilder->AssetAPI = ToolsContext->GetAssetAPI();
+	ToolsContext->ToolManager->RegisterToolType(TEXT("CreateActorSampleTool"), CreateActorSampleToolBuilder);
+
+	auto DrawCurveOnMeshSampleToolBuilder = NewObject<UDrawCurveOnMeshSampleToolBuilder>();
+	ToolsContext->ToolManager->RegisterToolType(TEXT("DrawCurveOnMeshSampleTool"), DrawCurveOnMeshSampleToolBuilder);
+
+	auto SurfacePointToolBuilder = NewObject<UMeshSurfacePointToolBuilder>();
+	ToolsContext->ToolManager->RegisterToolType(TEXT("SurfacePointTool"), SurfacePointToolBuilder);
+
+
+	// active tool type is not relevant here, we just set to default
+	ToolsContext->ToolManager->SelectActiveToolType(EToolSide::Left, TEXT("SurfacePointTool"));
+}
+
+
+
+
+void FSampleToolsEditorMode::Exit()
+{
+	// shutdown and clean up the ToolsContext
+	ToolsContext->ShutdownContext();
+	ToolsContext = nullptr;
+
+	if (Toolkit.IsValid())
+	{
+		FToolkitManager::Get().CloseToolkit(Toolkit.ToSharedRef());
+		Toolkit.Reset();
+	}
+
+	// Call base Exit method to ensure proper cleanup
+	FEdMode::Exit();
+}
+
+bool FSampleToolsEditorMode::UsesToolkits() const
+{
+	return true;
+}
+
+
+void FSampleToolsEditorMode::AddReferencedObjects(FReferenceCollector& Collector)
+{
+	Collector.AddReferencedObject(ToolsContext);
+}
+
+
+
+
+
+
+
+
+
+#undef LOCTEXT_NAMESPACE
\ No newline at end of file
diff --git a/Engine/Plugins/Experimental/SampleToolsEditorMode/Source/Private/SampleToolsEditorModeModule.cpp b/Engine/Plugins/Experimental/SampleToolsEditorMode/Source/Private/SampleToolsEditorModeModule.cpp
new file mode 100644
--- /dev/null
+++ b/Engine/Plugins/Experimental/SampleToolsEditorMode/Source/Private/SampleToolsEditorModeModule.cpp
@@ -0,0 +1,23 @@
+// Copyright 1998-2018 Epic Games, Inc. All Rights Reserved.
+
+#include "SampleToolsEditorModeModule.h"
+#include "SampleToolsEditorMode.h"
+
+#define LOCTEXT_NAMESPACE "FSampleToolsEditorModeModule"
+
+void FSampleToolsEditorModeModule::StartupModule()
+{
+	// This code will execute after your module is loaded into memory; the exact timing is specified in the .uplugin file per-module
+	FEditorModeRegistry::Get().RegisterMode<FSampleToolsEditorMode>(FSampleToolsEditorMode::EM_SampleToolsEditorModeId, LOCTEXT("SampleToolsEditorModeName", "SampleToolsEditorMode"), FSlateIcon(), true);
+}
+
+void FSampleToolsEditorModeModule::ShutdownModule()
+{
+	// This function may be called during shutdown to clean up your module.  For modules that support dynamic reloading,
+	// we call this function before unloading the module.
+	FEditorModeRegistry::Get().UnregisterMode(FSampleToolsEditorMode::EM_SampleToolsEditorModeId);
+}
+
+#undef LOCTEXT_NAMESPACE
+	
+IMPLEMENT_MODULE(FSampleToolsEditorModeModule, SampleToolsEditorMode)
\ No newline at end of file
diff --git a/Engine/Plugins/Experimental/SampleToolsEditorMode/Source/Private/SampleToolsEditorModeToolkit.cpp b/Engine/Plugins/Experimental/SampleToolsEditorMode/Source/Private/SampleToolsEditorModeToolkit.cpp
new file mode 100644
--- /dev/null
+++ b/Engine/Plugins/Experimental/SampleToolsEditorMode/Source/Private/SampleToolsEditorModeToolkit.cpp
@@ -0,0 +1,215 @@
+// Copyright 1998-2019 Epic Games, Inc. All Rights Reserved.
+
+#include "SampleToolsEditorModeToolkit.h"
+#include "SampleToolsEditorMode.h"
+#include "Engine/Selection.h"
+
+#include "Modules/ModuleManager.h"
+#include "PropertyEditorModule.h"
+#include "IDetailsView.h"
+
+#include "Widgets/Input/SButton.h"
+#include "Widgets/Text/STextBlock.h"
+
+
+
+#include "EditorModeManager.h"
+
+#define LOCTEXT_NAMESPACE "FSampleToolsEditorModeToolkit"
+
+FSampleToolsEditorModeToolkit::FSampleToolsEditorModeToolkit()
+{
+}
+
+void FSampleToolsEditorModeToolkit::Init(const TSharedPtr<IToolkitHost>& InitToolkitHost)
+{
+
+	FPropertyEditorModule& PropertyEditorModule = FModuleManager::GetModuleChecked<FPropertyEditorModule>("PropertyEditor");
+
+	FDetailsViewArgs DetailsViewArgs(
+		/*bUpdateFromSelection=*/ false,
+		/*bLockable=*/ false,
+		/*bAllowSearch=*/ false,
+		FDetailsViewArgs::HideNameArea,
+		/*bHideSelectionTip=*/ true,
+		/*InNotifyHook=*/ nullptr,
+		/*InSearchInitialKeyFocus=*/ false,
+		/*InViewIdentifier=*/ NAME_None);
+	DetailsViewArgs.DefaultsOnlyVisibility = EEditDefaultsOnlyNodeVisibility::Automatic;
+	DetailsViewArgs.bShowOptions = false;
+	DetailsViewArgs.bAllowMultipleTopLevelObjects = true;
+
+	DetailsView = PropertyEditorModule.CreateDetailView(DetailsViewArgs);
+
+
+	//////////////////////////////////////////////////////////////////////////
+	////////////////////////////////////////////////////////////////////////// 
+	// AddYourTool Step 3 - add a button to initialize your Tool
+	//////////////////////////////////////////////////////////////////////////
+	////////////////////////////////////////////////////////////////////////// 
+
+
+	SAssignNew(ToolkitWidget, SBorder)
+		.HAlign(HAlign_Center)
+		.Padding(25)
+		[
+			SNew(SVerticalBox)
+
+			+ SVerticalBox::Slot()
+				.AutoHeight()
+				.HAlign(HAlign_Center)
+				.Padding(50)
+				[
+					SNew(STextBlock)
+					.AutoWrapText(true)
+					.Text(LOCTEXT("HeaderLabel", "Sample Tools"))
+				]
+
+			+ SVerticalBox::Slot()
+				.HAlign(HAlign_Center)
+				.AutoHeight()
+				[
+					SNew(SButton).Text(LOCTEXT("CreateActorSampleToolLabel", "Create Actor on Click"))
+					.OnClicked_Lambda([this]() { return this->StartTool(TEXT("CreateActorSampleTool")); })
+					.IsEnabled_Lambda([this]() { return this->CanStartTool(TEXT("CreateActorSampleTool")); })
+				]
+
+
+			+ SVerticalBox::Slot()
+				.HAlign(HAlign_Center)
+				.AutoHeight()
+				[
+					SNew(SButton).Text(LOCTEXT("DrawCurveOnMeshSampleToolLabel", "Draw Curve On Mesh"))
+					.OnClicked_Lambda([this]() { return this->StartTool(TEXT("DrawCurveOnMeshSampleTool")); })
+					.IsEnabled_Lambda([this]() { return this->CanStartTool(TEXT("DrawCurveOnMeshSampleTool")); })
+				]
+
+
+			+ SVerticalBox::Slot()
+				.HAlign(HAlign_Center)
+				.AutoHeight()
+				[
+					SNew(SButton).Text(LOCTEXT("SurfacePointToolLabel", "Surface Point Tool"))
+					.OnClicked_Lambda([this]() { return this->StartTool(TEXT("SurfacePointTool")); })
+					.IsEnabled_Lambda([this]() { return this->CanStartTool(TEXT("SurfacePointTool")); })
+				]
+
+
+			+ SVerticalBox::Slot()
+				.HAlign(HAlign_Center)
+				.AutoHeight()
+				[
+				SNew(SHorizontalBox)
+				+ SHorizontalBox::Slot()
+					.AutoWidth()
+					[
+						SNew(SButton).Text(LOCTEXT("AcceptToolButtonLabel", "Accept"))
+						.OnClicked_Lambda([this]() { return this->EndTool(EToolShutdownType::Accept); })
+						.IsEnabled_Lambda([this]() { return this->CanAcceptActiveTool(); })
+					]
+				+ SHorizontalBox::Slot()
+					.AutoWidth()
+					[
+						SNew(SButton).Text(LOCTEXT("CancelToolButtonLabel", "Cancel"))
+						.OnClicked_Lambda([this]() { return this->EndTool(EToolShutdownType::Cancel); })
+						.IsEnabled_Lambda([this]() { return this->CanCancelActiveTool(); })
+					]
+				+ SHorizontalBox::Slot()
+					.AutoWidth()
+					[
+						SNew(SButton).Text(LOCTEXT("CompletedToolButtonLabel", "Complete"))
+						.OnClicked_Lambda([this]() { return this->EndTool(EToolShutdownType::Completed); })
+						.IsEnabled_Lambda([this]() { return this->CanCompleteActiveTool(); })
+					]
+				]
+
+			+ SVerticalBox::Slot()
+				.HAlign(HAlign_Center)
+				.Padding(2)
+				.AutoHeight()
+				.MaxHeight(500.0f)
+				[
+					DetailsView->AsShared()
+				]
+
+		];
+		
+	FModeToolkit::Init(InitToolkitHost);
+}
+
+
+bool FSampleToolsEditorModeToolkit::CanStartTool(const FString& ToolTypeIdentifier)
+{
+	UInteractiveToolManager* Manager = GetToolsEditorMode()->GetToolManager();
+
+	return (Manager->HasActiveTool(EToolSide::Left) == false) &&
+		(Manager->CanActivateTool(EToolSide::Left, ToolTypeIdentifier) == true);
+}
+
+bool FSampleToolsEditorModeToolkit::CanAcceptActiveTool()
+{
+	return GetToolsEditorMode()->GetToolManager()->CanAcceptActiveTool(EToolSide::Left);
+}
+
+bool FSampleToolsEditorModeToolkit::CanCancelActiveTool()
+{
+	return GetToolsEditorMode()->GetToolManager()->CanCancelActiveTool(EToolSide::Left);
+}
+
+bool FSampleToolsEditorModeToolkit::CanCompleteActiveTool()
+{
+	return GetToolsEditorMode()->GetToolManager()->HasActiveTool(EToolSide::Left) && CanCancelActiveTool() == false;
+}
+
+
+FReply FSampleToolsEditorModeToolkit::StartTool(const FString& ToolTypeIdentifier)
+{
+	if (GetToolsEditorMode()->GetToolManager()->SelectActiveToolType(EToolSide::Left, ToolTypeIdentifier) == false)
+	{
+		UE_LOG(LogTemp, Warning, TEXT("ToolManager: Unknown Tool Type %s"), *ToolTypeIdentifier);
+	}
+	else
+	{
+		UE_LOG(LogTemp, Warning, TEXT("ToolManager: Starting Tool Type %s"), *ToolTypeIdentifier);
+		GetToolsEditorMode()->GetToolManager()->ActivateTool(EToolSide::Left);
+
+		// Update properties panel
+		UInteractiveTool* CurTool = GetToolsEditorMode()->GetToolManager()->GetActiveTool(EToolSide::Left);
+		DetailsView->SetObjects(CurTool->GetToolProperties());
+	}
+	return FReply::Handled();
+}
+
+FReply FSampleToolsEditorModeToolkit::EndTool(EToolShutdownType ShutdownType)
+{
+	UE_LOG(LogTemp, Warning, TEXT("ENDING TOOL"));
+
+	GetToolsEditorMode()->GetToolManager()->DeactivateTool(EToolSide::Left, ShutdownType);
+
+	DetailsView->SetObject(nullptr);
+
+	return FReply::Handled();
+}
+
+
+FName FSampleToolsEditorModeToolkit::GetToolkitFName() const
+{
+	return FName("SampleToolsEditorMode");
+}
+
+FText FSampleToolsEditorModeToolkit::GetBaseToolkitName() const
+{
+	return NSLOCTEXT("SampleToolsEditorModeToolkit", "DisplayName", "SampleToolsEditorMode Tool");
+}
+
+class FEdMode* FSampleToolsEditorModeToolkit::GetEditorMode() const
+{
+	return GLevelEditorModeTools().GetActiveMode(FSampleToolsEditorMode::EM_SampleToolsEditorModeId);
+}
+
+FSampleToolsEditorMode* FSampleToolsEditorModeToolkit::GetToolsEditorMode() const
+{
+	return (FSampleToolsEditorMode*)GetEditorMode();
+}
+
+#undef LOCTEXT_NAMESPACE
diff --git a/Engine/Plugins/Experimental/SampleToolsEditorMode/Source/Public/SampleToolsEditorMode.h b/Engine/Plugins/Experimental/SampleToolsEditorMode/Source/Public/SampleToolsEditorMode.h
new file mode 100644
--- /dev/null
+++ b/Engine/Plugins/Experimental/SampleToolsEditorMode/Source/Public/SampleToolsEditorMode.h
@@ -0,0 +1,157 @@
+// Copyright 1998-2019 Epic Games, Inc. All Rights Reserved.
+
+#pragma once
+
+#include "CoreMinimal.h"
+#include "EdMode.h"
+
+#include "InputState.h"
+#include "InteractiveToolManager.h"
+#include "EdModeInteractiveToolsContext.h"
+
+
+class FEditorComponentSourceFactory;
+class FEditorToolAssetAPI;
+
+
+/**
+ * This class provides an example of how to connect an FEdMode to the 
+ * InteractiveTools framework. The various FEdMode input event handlers 
+ * forward events to a UEdModeInteractiveToolsContext instance, which 
+ * has all the logic for interacting with the InputRouter, ToolManager, etc.
+ * The functions here are basically just forwarders.
+ */
+class FSampleToolsEditorMode : public FEdMode
+{
+public:
+	const static FEditorModeID EM_SampleToolsEditorModeId;
+public:
+	FSampleToolsEditorMode();
+	virtual ~FSampleToolsEditorMode();
+
+	////////////////
+	// FEdMode interface
+	////////////////
+
+	virtual void Tick(FEditorViewportClient* ViewportClient, float DeltaTime) override;
+	virtual void Render(const FSceneView* View, FViewport* Viewport, FPrimitiveDrawInterface* PDI) override;
+	virtual void ActorSelectionChangeNotify() override;
+	virtual bool UsesToolkits() const override;
+
+	virtual void AddReferencedObjects(FReferenceCollector& Collector) override;
+
+	// these disable the standard gizmo, which is probably want we want in
+	// these tools as we can't hit-test the standard gizmo...
+	virtual bool AllowWidgetMove() override { return false; }
+	virtual bool ShouldDrawWidget() const override { return false; }
+	virtual bool UsesTransformWidget() const override { return false; }
+
+
+	/*
+	 * focus events
+	 */
+
+	// called when we "start" this editor mode (ie switch to this tab)
+	virtual void Enter() override;
+
+	// called when we "end" this editor mode (ie switch to another tab)
+	virtual void Exit() override;
+
+	// called when viewport window is focused
+	virtual bool ReceivedFocus(FEditorViewportClient* ViewportClient, FViewport* Viewport) override;
+
+	// called when viewport window loses focus (ie when some other window is focused)
+	// *not* called when editor is backgrounded, but is called when editor is minimized
+	virtual bool LostFocus(FEditorViewportClient* ViewportClient, FViewport* Viewport) override;
+
+
+	/*
+	 * Mouse position events - These are called when no mouse button is down
+	 */
+
+	// called when mouse moves over viewport window
+	virtual bool MouseEnter(FEditorViewportClient* ViewportClient, FViewport* Viewport, int32 x, int32 y) override;
+
+	// called when mouse leaves viewport window
+	virtual bool MouseLeave(FEditorViewportClient* ViewportClient, FViewport* Viewport) override;
+
+	// called on any mouse-move event. *not* called during tracking/capturing, eg if any button is down
+	virtual bool MouseMove(FEditorViewportClient* ViewportClient, FViewport* Viewport, int32 x, int32 y) override;
+
+
+	/*
+	 * Input Button/Axis Events & Mouse Capture
+	 *
+	 *   event sequence for left mouse down/capture/up is:
+	 *     - InputKey( EKeys::LeftMouseButton, IE_PRESSED )
+	 *     - StartTracking()
+	 *     - CapturedMouseMove() and InputAxis() (repeated for each mouse-move)
+	 *     - Inputkey( EKeys::LeftMouseButton, IE_RELEASED )
+	 *     - EndTracking()
+	 *
+	 *   for doubleclick, we get one of the above sequence, then a second
+	 *   where instead of the first IE_PRESSED we get a IE_DoubleClick
+	 *
+	 *   for mouse wheel, we get following sequence
+	 *     - InputKey ( EKeys::MouseScrollUp / Down, IE_PRESSED )
+	 *     - InputAxis() for wheel move
+	 *     - InputKey ( EKeys::MouseScrollUp / Down, IE_RELEASED )
+	 *   it appears that there will only ever be *one* InputAxis() between the pressed/released sequenece
+	 *
+	 *   Note that this wheel event sequence can happen *during* a 
+	 *   middle-mouse tracking sequence on a wheel-button (nice!)
+	 */
+
+
+	// This is not just called for keyboard keys, it is also called for mouse down/up events!
+	// Eg for left-press we get Key = EKeys::LeftMouseButton and Event = IE_Pressed
+	// Return value indicates "handled'. If we return true for mouse press events then
+	// the StartTracking/CapturedMouseMove/EndTracking sequence is *not* called (but we
+	// also don't get MouseMove(), so the mouse-movement events appear to be lost?)
+	virtual bool InputKey(FEditorViewportClient* ViewportClient, FViewport* Viewport, FKey Key, EInputEvent Event) override;
+
+
+	// Called for 1D Axis movements - EKeys::MouseX, EKeys::MouseY, and EKeys::MouseWheelAxis
+	// Called even if we return true from InputKey which otherwise blocks tracking. 
+	// Return value indicates "handled" but has no effect on CapturedMouseMove events (ie we still get them)
+	virtual bool InputAxis(FEditorViewportClient* InViewportClient, FViewport* Viewport, int32 ControllerId, FKey Key, float Delta, float DeltaTime);
+
+
+	// called on mouse-down. return value indicates whether this was "handled" but
+	// does not mean we get exclusive capture events!
+	virtual bool StartTracking(FEditorViewportClient* InViewportClient, FViewport* InViewport) override;
+
+	// called during mouse-down mouse-move. Always called, return value is used to indicate whether
+	// we "handled" this mouse move (see FEditorModeTools::CapturedMouseMove) but does not pre-empt
+	// other editor modes from seeing this move event...
+	virtual bool CapturedMouseMove(FEditorViewportClient* InViewportClient, FViewport* InViewport, int32 InMouseX, int32 InMouseY) override;
+
+	// always called on mouse-up
+	virtual bool EndTracking(FEditorViewportClient* InViewportClient, FViewport* InViewport) override;
+
+
+	//////////////////
+	// End of FEdMode interface
+	//////////////////
+
+
+public:
+
+	/**
+	 * @return active ToolManager
+	 */
+	virtual UInteractiveToolManager* GetToolManager() const
+	{ 
+		return ToolsContext->ToolManager;
+	}
+
+protected:
+
+	/**
+	 * This is the ToolsContext adapter that attaches to the FEdMode. 
+	 * We forward FEdMode event handler functions to this implementation.
+	 * This UObject is explicitly kept alive by adding it to the Collector in ::AddReferencedObjects()
+	 */
+	UEdModeInteractiveToolsContext* ToolsContext;
+
+};
diff --git a/Engine/Plugins/Experimental/SampleToolsEditorMode/Source/Public/SampleToolsEditorModeModule.h b/Engine/Plugins/Experimental/SampleToolsEditorMode/Source/Public/SampleToolsEditorModeModule.h
new file mode 100644
--- /dev/null
+++ b/Engine/Plugins/Experimental/SampleToolsEditorMode/Source/Public/SampleToolsEditorModeModule.h
@@ -0,0 +1,15 @@
+// Copyright 1998-2018 Epic Games, Inc. All Rights Reserved.
+
+#pragma once
+
+#include "CoreMinimal.h"
+#include "Modules/ModuleManager.h"
+
+class FSampleToolsEditorModeModule : public IModuleInterface
+{
+public:
+
+	/** IModuleInterface implementation */
+	virtual void StartupModule() override;
+	virtual void ShutdownModule() override;
+};
diff --git a/Engine/Plugins/Experimental/SampleToolsEditorMode/Source/Public/SampleToolsEditorModeToolkit.h b/Engine/Plugins/Experimental/SampleToolsEditorMode/Source/Public/SampleToolsEditorModeToolkit.h
new file mode 100644
--- /dev/null
+++ b/Engine/Plugins/Experimental/SampleToolsEditorMode/Source/Public/SampleToolsEditorModeToolkit.h
@@ -0,0 +1,46 @@
+// Copyright 1998-2019 Epic Games, Inc. All Rights Reserved.
+
+#pragma once
+
+#include "CoreMinimal.h"
+#include "Toolkits/BaseToolkit.h"
+#include "InteractiveTool.h"
+#include "SampleToolsEditorMode.h"
+
+class IDetailsView;
+
+/**
+ * This FModeToolkit just creates a basic UI panel that allows various InteractiveTools to
+ * be initialized, and a DetailsView used to show properties of the active Tool.
+ */
+class FSampleToolsEditorModeToolkit : public FModeToolkit
+{
+public:
+	FSampleToolsEditorModeToolkit();
+	
+	// FModeToolkit interface 
+	virtual void Init(const TSharedPtr<IToolkitHost>& InitToolkitHost) override;
+
+	// IToolkit interface 
+	virtual FName GetToolkitFName() const override;
+	virtual FText GetBaseToolkitName() const override;
+	virtual class FEdMode* GetEditorMode() const override;
+	virtual TSharedPtr<class SWidget> GetInlineContent() const override { return ToolkitWidget; }
+
+	virtual FSampleToolsEditorMode* GetToolsEditorMode() const;
+
+private:
+
+	TSharedPtr<SWidget> ToolkitWidget;
+	TSharedPtr<IDetailsView> DetailsView;
+
+	// these functions just forward calls to the ToolsContext / ToolManager
+
+	bool CanStartTool(const FString& ToolTypeIdentifier);
+	bool CanAcceptActiveTool();
+	bool CanCancelActiveTool();
+	bool CanCompleteActiveTool();
+
+	FReply StartTool(const FString& ToolTypeIdentifier);
+	FReply EndTool(EToolShutdownType ShutdownType);
+};
diff --git a/Engine/Plugins/Experimental/SampleToolsEditorMode/Source/SampleToolsEditorMode.Build.cs b/Engine/Plugins/Experimental/SampleToolsEditorMode/Source/SampleToolsEditorMode.Build.cs
new file mode 100644
--- /dev/null
+++ b/Engine/Plugins/Experimental/SampleToolsEditorMode/Source/SampleToolsEditorMode.Build.cs
@@ -0,0 +1,60 @@
+// Copyright 1998-2019 Epic Games, Inc. All Rights Reserved.
+
+using UnrealBuildTool;
+
+public class SampleToolsEditorMode : ModuleRules
+{
+	public SampleToolsEditorMode(ReadOnlyTargetRules Target) : base(Target)
+	{
+		PCHUsage = ModuleRules.PCHUsageMode.UseExplicitOrSharedPCHs;
+		
+		PublicIncludePaths.AddRange(
+			new string[] {
+                //"ContentBrowser"
+				// ... add public include paths required here ...
+			}
+			);
+				
+		
+		PrivateIncludePaths.AddRange(
+			new string[] {
+				// ... add other private include paths required here ...
+			}
+			);
+			
+		
+		PublicDependencyModuleNames.AddRange(
+			new string[]
+			{
+				"Core"
+				// ... add other public dependencies that you statically link with here ...
+			}
+			);
+			
+		
+		PrivateDependencyModuleNames.AddRange(
+			new string[]
+			{
+				"CoreUObject",
+				"Engine",
+				"Slate",
+				"SlateCore",
+				"InputCore",
+				"UnrealEd",
+                "ContentBrowser",
+                "LevelEditor",
+                "InteractiveToolsFramework",
+				"EditorInteractiveToolsFramework"
+				// ... add private dependencies that you statically link with here ...	
+			}
+            );
+		
+		
+		DynamicallyLoadedModuleNames.AddRange(
+			new string[]
+			{
+				// ... add any modules that your module loads dynamically here ...
+			}
+			);
+	}
+}
