add missing includes
#rnx

diff --git a/Engine/Plugins/Experimental/SampleToolsEditorMode/Source/Private/SampleTools/CreateActorSampleTool.cpp b/Engine/Plugins/Experimental/SampleToolsEditorMode/Source/Private/SampleTools/CreateActorSampleTool.cpp
--- a/Engine/Plugins/Experimental/SampleToolsEditorMode/Source/Private/SampleTools/CreateActorSampleTool.cpp
+++ b/Engine/Plugins/Experimental/SampleToolsEditorMode/Source/Private/SampleTools/CreateActorSampleTool.cpp
@@ -2,6 +2,8 @@
 
 #include "CreateActorSampleTool.h"
 #include "ToolBuilderUtil.h"
+#include "CollisionQueryParams.h"
+#include "Engine/World.h"
 
 // localization namespace
 #define LOCTEXT_NAMESPACE "UCreateActorSampleTool"
diff --git a/Engine/Plugins/Experimental/SampleToolsEditorMode/Source/Private/SampleTools/DrawCurveOnMeshSampleTool.cpp b/Engine/Plugins/Experimental/SampleToolsEditorMode/Source/Private/SampleTools/DrawCurveOnMeshSampleTool.cpp
--- a/Engine/Plugins/Experimental/SampleToolsEditorMode/Source/Private/SampleTools/DrawCurveOnMeshSampleTool.cpp
+++ b/Engine/Plugins/Experimental/SampleToolsEditorMode/Source/Private/SampleTools/DrawCurveOnMeshSampleTool.cpp
@@ -2,6 +2,7 @@
 
 #include "DrawCurveOnMeshSampleTool.h"
 #include "ToolBuilderUtil.h"
+#include "SceneManagement.h"   // FPrimitiveDrawInterface
 
 // localization namespace
 #define LOCTEXT_NAMESPACE "UDrawCurveOnMeshSampleTool"
@@ -47,7 +48,7 @@ void UDrawCurveOnMeshSampleTool::Setup()
 
 void UDrawCurveOnMeshSampleTool::Render(IToolsContextRenderAPI* RenderAPI)
 {
-	auto PDI = RenderAPI->GetPrimitiveDrawInterface();
+	FPrimitiveDrawInterface* PDI = RenderAPI->GetPrimitiveDrawInterface();
 	
 	int NumPts = Positions.Num();
 	for (int i = 0; i < NumPts - 1; ++i)
