fix copyright notice
#rnx

diff --git a/Engine/Plugins/Experimental/SampleToolsEditorMode/Source/Private/SampleToolsEditorModeModule.cpp b/Engine/Plugins/Experimental/SampleToolsEditorMode/Source/Private/SampleToolsEditorModeModule.cpp
--- a/Engine/Plugins/Experimental/SampleToolsEditorMode/Source/Private/SampleToolsEditorModeModule.cpp
+++ b/Engine/Plugins/Experimental/SampleToolsEditorMode/Source/Private/SampleToolsEditorModeModule.cpp
@@ -1,4 +1,4 @@
-// Copyright 1998-2018 Epic Games, Inc. All Rights Reserved.
+// Copyright 1998-2019 Epic Games, Inc. All Rights Reserved.
 
 #include "SampleToolsEditorModeModule.h"
 #include "SampleToolsEditorMode.h"
diff --git a/Engine/Plugins/Experimental/SampleToolsEditorMode/Source/Public/SampleToolsEditorModeModule.h b/Engine/Plugins/Experimental/SampleToolsEditorMode/Source/Public/SampleToolsEditorModeModule.h
--- a/Engine/Plugins/Experimental/SampleToolsEditorMode/Source/Public/SampleToolsEditorModeModule.h
+++ b/Engine/Plugins/Experimental/SampleToolsEditorMode/Source/Public/SampleToolsEditorModeModule.h
@@ -1,4 +1,4 @@
-// Copyright 1998-2018 Epic Games, Inc. All Rights Reserved.
+// Copyright 1998-2019 Epic Games, Inc. All Rights Reserved.
 
 #pragma once
 
