add UInteractiveToolPropertySet, InteractiveTool automatically listens to change delegate on this object. InteractiveTools should implement property sets as subclasses of UInteractiveToolPropertySet  (not currently enforced)

diff --git a/Engine/Plugins/Experimental/SampleToolsEditorMode/Source/Private/SampleTools/DrawCurveOnMeshSampleTool.cpp b/Engine/Plugins/Experimental/SampleToolsEditorMode/Source/Private/SampleTools/DrawCurveOnMeshSampleTool.cpp
--- a/Engine/Plugins/Experimental/SampleToolsEditorMode/Source/Private/SampleTools/DrawCurveOnMeshSampleTool.cpp
+++ b/Engine/Plugins/Experimental/SampleToolsEditorMode/Source/Private/SampleTools/DrawCurveOnMeshSampleTool.cpp
@@ -24,8 +24,7 @@ UMeshSurfacePointTool* UDrawCurveOnMeshSampleToolBuilder::CreateNewTool(const FT
  * Tool
  */
 
-
-UDrawCurveOnMeshSampleTool::UDrawCurveOnMeshSampleTool()
+UDrawCurveOnMeshSampleToolProperties::UDrawCurveOnMeshSampleToolProperties()
 {
 	Thickness = 4.0f;
 	DepthBias = 0.0f;
@@ -36,12 +35,18 @@ UDrawCurveOnMeshSampleTool::UDrawCurveOnMeshSampleTool()
 }
 
 
+UDrawCurveOnMeshSampleTool::UDrawCurveOnMeshSampleTool()
+{
+}
+
+
 void UDrawCurveOnMeshSampleTool::Setup()
 {
 	UMeshSurfacePointTool::Setup();
 
-	// provide self as property object
-	ToolPropertyObjects.Add(this);
+	// add settings object
+	Settings = NewObject<UDrawCurveOnMeshSampleToolProperties>(this, TEXT("Settings"));
+	AddToolPropertySource(Settings);
 }
 
 
@@ -53,9 +58,9 @@ void UDrawCurveOnMeshSampleTool::Render(IToolsContextRenderAPI* RenderAPI)
 	int NumPts = Positions.Num();
 	for (int i = 0; i < NumPts - 1; ++i)
 	{
-		FVector A = Positions[i] + NormalOffset*Normals[i];
-		FVector B = Positions[i+1] + NormalOffset*Normals[i+1];
-		PDI->DrawLine(A, B, Color, 0, Thickness, DepthBias, bScreenSpace);
+		FVector A = Positions[i] + Settings->NormalOffset*Normals[i];
+		FVector B = Positions[i+1] + Settings->NormalOffset*Normals[i+1];
+		PDI->DrawLine(A, B, Settings->Color, 0, Settings->Thickness, Settings->DepthBias, Settings->bScreenSpace);
 	}
 }
 
@@ -81,7 +86,7 @@ void UDrawCurveOnMeshSampleTool::OnUpdateDrag(const FRay& Ray)
 	FHitResult OutHit;
 	if (HitTest(Ray, OutHit))
 	{
-		if ( FVector::Dist(OutHit.ImpactPoint, Positions[Positions.Num()-1]) > MinSpacing)
+		if ( FVector::Dist(OutHit.ImpactPoint, Positions[Positions.Num()-1]) > Settings->MinSpacing)
 		{
 			Positions.Add(OutHit.ImpactPoint);
 			Normals.Add(OutHit.ImpactNormal);
diff --git a/Engine/Plugins/Experimental/SampleToolsEditorMode/Source/Private/SampleTools/DrawCurveOnMeshSampleTool.h b/Engine/Plugins/Experimental/SampleToolsEditorMode/Source/Private/SampleTools/DrawCurveOnMeshSampleTool.h
--- a/Engine/Plugins/Experimental/SampleToolsEditorMode/Source/Private/SampleTools/DrawCurveOnMeshSampleTool.h
+++ b/Engine/Plugins/Experimental/SampleToolsEditorMode/Source/Private/SampleTools/DrawCurveOnMeshSampleTool.h
@@ -12,7 +12,7 @@
 /**
  * UMeshSurfacePointToolBuilder override for UDrawCurveOnMeshSampleTool
  */
-UCLASS()
+UCLASS(Transient)
 class SAMPLETOOLSEDITORMODE_API UDrawCurveOnMeshSampleToolBuilder : public UMeshSurfacePointToolBuilder
 {
 	GENERATED_BODY()
@@ -22,13 +22,44 @@ public:
 };
 
 
+/**
+ * Settings UObject for UDrawCurveOnMeshSampleTool. This UClass inherits from UInteractiveToolPropertySet,
+ * which provides an OnModified delegate that the Tool will listen to for changes in property values.
+ */
+UCLASS(Transient)
+class SAMPLETOOLSEDITORMODE_API UDrawCurveOnMeshSampleToolProperties : public UInteractiveToolPropertySet
+{
+	GENERATED_BODY()
+public:
+	UDrawCurveOnMeshSampleToolProperties();
+
+	UPROPERTY(EditAnywhere, Category = Options)
+	FLinearColor Color;
+
+	UPROPERTY(EditAnywhere, Category = Options, meta = (DisplayName = "Thickness", UIMin = "0.25", UIMax = "10.0", ClampMin = "0.01", ClampMax = "1000.0"))
+	float Thickness;
+
+	UPROPERTY(EditAnywhere, Category = Options, meta = (DisplayName = "Min Spacing", UIMin = "0.01", UIMax = "10.0"))
+	float MinSpacing;
+
+	UPROPERTY(EditAnywhere, Category = Options, meta = (DisplayName = "Offset", UIMin = "0.0", UIMax = "10.0", ClampMin = "-1000.0", ClampMax = "1000.0"))
+	float NormalOffset;
+
+	UPROPERTY(EditAnywhere, Category = Options, meta = (DisplayName = "Depth Bias", UIMin = "-10.0", UIMax = "10.0"))
+	float DepthBias;
+
+	UPROPERTY(EditAnywhere, Category = Options)
+	bool bScreenSpace;
+};
+
+
 
 /**
  * UDrawCurveOnMeshSampleTool is a sample Tool that allows the user to draw curves on the surface of
  * a selected Mesh Component. The various rendering properties of the polycurve are exposed and can be tweaked.
  * Nothing is done with the curve, it is just drawn by ::Render() and discarded when the Tool exits.
  */
-UCLASS()
+UCLASS(Transient)
 class SAMPLETOOLSEDITORMODE_API UDrawCurveOnMeshSampleTool : public UMeshSurfacePointTool
 {
 	GENERATED_BODY()
@@ -48,23 +79,8 @@ public:
 	virtual void OnEndDrag(const FRay& Ray) override;
 
 protected:
-	UPROPERTY(EditAnywhere, Category = Options)
-	FLinearColor Color;
-
-	UPROPERTY(EditAnywhere, Category = Options, meta = (DisplayName = "Thickness", UIMin = "0.25", UIMax = "10.0", ClampMin = "0.01", ClampMax = "1000.0"))
-	float Thickness;
-
-	UPROPERTY(EditAnywhere, Category = Options, meta = (DisplayName = "Min Spacing", UIMin = "0.01", UIMax = "10.0"))
-	float MinSpacing;
-
-	UPROPERTY(EditAnywhere, Category = Options, meta = (DisplayName = "Offset", UIMin = "0.0", UIMax = "10.0", ClampMin = "-1000.0", ClampMax = "1000.0"))
-	float NormalOffset;
-
-	UPROPERTY(EditAnywhere, Category = Options, meta = (DisplayName = "Depth Bias", UIMin = "-10.0", UIMax = "10.0"))
-	float DepthBias;
-
-	UPROPERTY(EditAnywhere, Category = Options)
-	bool bScreenSpace;
+	UPROPERTY()
+	UDrawCurveOnMeshSampleToolProperties* Settings;
 
 	UPROPERTY()
 	TArray<FVector> Positions;
diff --git a/Engine/Source/Runtime/InteractiveToolsFramework/Private/InteractiveTool.cpp b/Engine/Source/Runtime/InteractiveToolsFramework/Private/InteractiveTool.cpp
--- a/Engine/Source/Runtime/InteractiveToolsFramework/Private/InteractiveTool.cpp
+++ b/Engine/Source/Runtime/InteractiveToolsFramework/Private/InteractiveTool.cpp
@@ -47,6 +47,18 @@ void UInteractiveTool::AddToolPropertySource(UObject* PropertyObject)
 	ToolPropertyObjects.Add(PropertyObject);
 }
 
+void UInteractiveTool::AddToolPropertySource(UInteractiveToolPropertySet* PropertySet)
+{
+	check(ToolPropertyObjects.Contains(PropertySet) == false);
+	ToolPropertyObjects.Add(PropertySet);
+	// @todo do we need to create a lambda every time for this?
+	PropertySet->GetOnModified().AddLambda([this](UObject* PropertySet, UProperty* Property)
+	{
+		OnPropertyModified(PropertySet, Property);
+	});
+}
+
+
 const TArray<UObject*>& UInteractiveTool::GetToolProperties() const
 {
 	return ToolPropertyObjects;
diff --git a/Engine/Source/Runtime/InteractiveToolsFramework/Public/InteractiveTool.h b/Engine/Source/Runtime/InteractiveToolsFramework/Public/InteractiveTool.h
--- a/Engine/Source/Runtime/InteractiveToolsFramework/Public/InteractiveTool.h
+++ b/Engine/Source/Runtime/InteractiveToolsFramework/Public/InteractiveTool.h
@@ -22,6 +22,44 @@ enum class EToolShutdownType
 };
 
 
+/** This delegate is used by UInteractiveToolPropertySet */
+DECLARE_MULTICAST_DELEGATE_TwoParams(FInteractiveToolPropertySetModifiedSignature, UObject*, UProperty*);
+
+
+/**
+ * a UInteractiveTool contains a set of UObjects that contain "properties" of the Tool, ie
+ * the configuration flags, parameters, etc that control the Tool. Currently any UObject
+ * can be added as a property set, however there is no automatic mechanism for those child 
+ * UObjects to notify the Tool when a property changes.
+ * 
+ * If you make your property set UObjects subclasses of UInteractiveToolPropertySet, then
+ * when the properties are changed in the Editor, the parent Tool will be automatically notified.
+ * You can override UInteractiveTool::OnPropertyModified() to act on these notifications
+ */
+UCLASS(Transient)
+class INTERACTIVETOOLSFRAMEWORK_API UInteractiveToolPropertySet : public UObject
+{
+	GENERATED_BODY()
+
+protected:
+	FInteractiveToolPropertySetModifiedSignature OnModified;
+
+public:
+
+	/** @return the multicast delegate that is called when properties are modified */
+	FInteractiveToolPropertySetModifiedSignature& GetOnModified()
+	{
+		return OnModified;
+	}
+
+	/** posts a message to the OnModified delegate with the modified UProperty */
+	void PostEditChangeProperty(FPropertyChangedEvent& PropertyChangedEvent)
+	{
+		OnModified.Broadcast(this, PropertyChangedEvent.Property);
+	}
+};
+
+
 /**
  * UInteractiveTool is the base class for all Tools in the InteractiveToolsFramework.
  * A Tool is is a "lightweight mode" that may "own" one or more Actors/Components/etc in
@@ -123,6 +161,15 @@ public:
 	virtual const TArray<UObject*>& GetToolProperties() const;
 
 
+	/**
+	 * Automatically called by UInteractiveToolPropertySet.OnModified delegate to notify Tool of child property set changes
+	 * @param PropertySet which UInteractiveToolPropertySet was modified
+	 * @param Property which UProperty in the set was modified
+	 */
+	virtual void OnPropertyModified(UObject* PropertySet, UProperty* Property)
+	{
+	}
+
 
 protected:
 
@@ -138,6 +185,12 @@ protected:
 	 * @param Property object to add
 	 */
 	virtual void AddToolPropertySource(UObject* PropertyObject);
+
+	/**
+	 * Add a PropertySet object for this Tool
+	 * @param PropertySet Property Set object to add
+	 */
+	virtual void AddToolPropertySource(UInteractiveToolPropertySet* PropertySet);
 };
 
 
