add ClickDragBehavior, ClickDragTool, and MeasureDistanceSampleTool. Update existing Sample Tools to use new property sets design. Modifier keys are now updated on captured mouse move.
#rnx

diff --git a/Engine/Plugins/Experimental/SampleToolsEditorMode/Source/Private/SampleTools/CreateActorSampleTool - Copy.cpp b/Engine/Plugins/Experimental/SampleToolsEditorMode/Source/Private/SampleTools/CreateActorSampleTool - Copy.cpp
new file mode 100644
--- /dev/null
+++ b/Engine/Plugins/Experimental/SampleToolsEditorMode/Source/Private/SampleTools/CreateActorSampleTool - Copy.cpp	
@@ -0,0 +1,165 @@
+// Copyright 1998-2019 Epic Games, Inc. All Rights Reserved.
+
+#include "MeasureDistanceSampleTool.h"
+#include "InteractiveToolManager.h"
+#include "ToolBuilderUtil.h"
+#include "BaseBehaviors/ClickDragBehavior.h"
+
+// for raycast into World
+#include "CollisionQueryParams.h"
+#include "Engine/World.h"
+
+// localization namespace
+#define LOCTEXT_NAMESPACE "UMeasureDistanceSampleTool"
+
+/*
+ * ToolBuilder
+ */
+
+
+bool UMeasureDistanceSampleToolBuilder::CanBuildTool(const FToolBuilderState & SceneState) const
+{
+	return true;
+}
+
+UInteractiveTool* UMeasureDistanceSampleToolBuilder::BuildTool(const FToolBuilderState & SceneState) const
+{
+	UMeasureDistanceSampleTool* NewTool = NewObject<UMeasureDistanceSampleTool>(SceneState.ToolManager);
+	NewTool->SetWorld(SceneState.World);
+	return NewTool;
+}
+
+
+
+/*
+ * Tool
+ */
+
+UMeasureDistanceProperties::UMeasureDistanceProperties()
+{
+	// initialize the points and distance to reasonable values
+	StartPoint = FVector(0,0,0);
+	EndPoint = FVector(0,0,100);
+	Distance = 100;
+}
+
+
+UMeasureDistanceSampleTool::UMeasureDistanceSampleTool()
+{
+}
+
+
+void UMeasureDistanceSampleTool::SetWorld(UWorld* World)
+{
+	check(World);
+	this->TargetWorld = World;
+}
+
+
+
+
+void UMeasureDistanceSampleTool::Setup()
+{
+	UInteractiveTool::Setup();
+
+	// Add default mouse input behavior
+	UClickDragInputBehavior* MouseBehavior = NewObject<UClickDragInputBehavior>();
+	// We will use the shift key to indicate that we should move the second point. 
+	// This call tells the Behavior to call our OnUpdateModifierState() function on mouse-down and mouse-move
+	MouseBehavior->Modifiers.RegisterModifier(MoveSecondPointModifierID, FInputDeviceState::IsShiftKeyDown);
+	MouseBehavior->Initialize(this);
+	AddInputBehavior(MouseBehavior);
+
+	// Create the property set and register it with the Tool
+	Properties = NewObject<UMeasureDistanceProperties>(this, "Measurement");
+	AddToolPropertySource(Properties);
+	
+	bSecondPointModifierDown = false;
+	bMoveSecondPoint = false;
+}
+
+
+void UMeasureDistanceSampleTool::OnUpdateModifierState(int ModifierID, bool bIsOn)
+{
+	// keep track of the "second point" modifier (shift key for mouse input)
+	if (ModifierID == MoveSecondPointModifierID)
+	{
+		bSecondPointModifierDown = bIsOn;
+	}
+}
+
+
+
+bool UMeasureDistanceSampleTool::CanBeginClickDragSequence(const FInputDeviceRay& PressPos)
+{
+	// we only start drag if press-down is on top of something we can raycast
+	FVector Temp;
+	return FindRayHit(PressPos.WorldRay, Temp);
+}
+
+void UMeasureDistanceSampleTool::OnClickPress(const FInputDeviceRay& PressPos)
+{
+	// determine whether we are moving first or second point for the drag sequence
+	bMoveSecondPoint = bSecondPointModifierDown;
+	UpdatePosition(PressPos.WorldRay);
+}
+
+void UMeasureDistanceSampleTool::OnClickDrag(const FInputDeviceRay& DragPos)
+{
+	UpdatePosition(DragPos.WorldRay);
+}
+
+
+
+bool UMeasureDistanceSampleTool::FindRayHit(const FRay& WorldRay, FVector& HitPos)
+{
+	// trace a ray into the World
+	FCollisionObjectQueryParams QueryParams(FCollisionObjectQueryParams::AllObjects);
+	FHitResult Result;
+	bool bHitWorld = TargetWorld->LineTraceSingleByObjectType(Result, WorldRay.Origin, WorldRay.PointAt(999999), QueryParams);
+	if (bHitWorld)
+	{
+		HitPos = Result.ImpactPoint;
+		return true;
+	}
+	return false;
+}
+
+
+void UMeasureDistanceSampleTool::UpdatePosition(const FRay& WorldRay)
+{
+	if (FindRayHit(WorldRay, (bMoveSecondPoint) ? Properties->EndPoint : Properties->StartPoint))
+	{
+		UpdateDistance();
+	}
+}
+
+
+void UMeasureDistanceSampleTool::UpdateDistance()
+{
+	Properties->Distance = FVector::Distance(Properties->StartPoint, Properties->EndPoint);
+}
+
+
+void UMeasureDistanceSampleTool::OnPropertyModified(UObject* PropertySet, UProperty* Property)
+{
+	// if the user updated any of the property fields, update the distance
+	UpdateDistance();
+}
+
+
+
+void UMeasureDistanceSampleTool::Render(IToolsContextRenderAPI* RenderAPI)
+{
+	FPrimitiveDrawInterface* PDI = RenderAPI->GetPrimitiveDrawInterface();
+	// draw a thin line that shows through objects
+	PDI->DrawLine(Properties->StartPoint, Properties->EndPoint,
+		FColor(240, 16, 16), SDPG_Foreground, 2.0f, 0.0f, true);
+	// draw a thicker line that is depth-tested
+	PDI->DrawLine(Properties->StartPoint, Properties->EndPoint,
+		FColor(240, 16, 16), SDPG_World, 4.0f, 0.0f, true);
+}
+
+
+
+#undef LOCTEXT_NAMESPACE
\ No newline at end of file
diff --git a/Engine/Plugins/Experimental/SampleToolsEditorMode/Source/Private/SampleTools/CreateActorSampleTool.cpp b/Engine/Plugins/Experimental/SampleToolsEditorMode/Source/Private/SampleTools/CreateActorSampleTool.cpp
--- a/Engine/Plugins/Experimental/SampleToolsEditorMode/Source/Private/SampleTools/CreateActorSampleTool.cpp
+++ b/Engine/Plugins/Experimental/SampleToolsEditorMode/Source/Private/SampleTools/CreateActorSampleTool.cpp
@@ -1,6 +1,7 @@
 // Copyright 1998-2019 Epic Games, Inc. All Rights Reserved.
 
 #include "CreateActorSampleTool.h"
+#include "InteractiveToolManager.h"
 #include "ToolBuilderUtil.h"
 #include "CollisionQueryParams.h"
 #include "Engine/World.h"
@@ -20,7 +21,7 @@ bool UCreateActorSampleToolBuilder::CanBuildTool(const FToolBuilderState & Scene
 
 UInteractiveTool* UCreateActorSampleToolBuilder::BuildTool(const FToolBuilderState & SceneState) const
 {
-	UCreateActorSampleTool* NewTool = NewObject<UCreateActorSampleTool>();
+	UCreateActorSampleTool* NewTool = NewObject<UCreateActorSampleTool>(SceneState.ToolManager);
 	NewTool->SetWorld(SceneState.World);
 	NewTool->SetAssetAPI(AssetAPI);
 	return NewTool;
@@ -33,13 +34,18 @@ UInteractiveTool* UCreateActorSampleToolBuilder::BuildTool(const FToolBuilderSta
  */
 
 
-UCreateActorSampleTool::UCreateActorSampleTool()
+UCreateActorSampleToolProperties::UCreateActorSampleToolProperties()
 {
 	PlaceOnObjects = true;
 	GroundHeight = 0.0f;
 }
 
 
+UCreateActorSampleTool::UCreateActorSampleTool()
+{
+}
+
+
 void UCreateActorSampleTool::SetWorld(UWorld* World)
 {
 	this->TargetWorld = World;
@@ -55,8 +61,8 @@ void UCreateActorSampleTool::Setup()
 {
 	USingleClickTool::Setup();
 
-	// provide self as property object
-	ToolPropertyObjects.Add(this);
+	Properties = NewObject<UCreateActorSampleToolProperties>(this);
+	AddToolPropertySource(Properties);
 }
 
 
@@ -73,14 +79,14 @@ void UCreateActorSampleTool::OnClicked(const FInputDeviceRay& ClickPos)
 	FCollisionObjectQueryParams QueryParams(FCollisionObjectQueryParams::AllObjects);
 	FHitResult Result;
 	bool bHitWorld = TargetWorld->LineTraceSingleByObjectType(Result, RayStart, RayEnd, QueryParams);
-	if (PlaceOnObjects && bHitWorld)
+	if (Properties->PlaceOnObjects && bHitWorld)
 	{
 		NewActorPos = Result.ImpactPoint;
 	}
 	else
 	{
 		// hit the ground plane
-		FPlane GroundPlane(FVector(0,0,GroundHeight) , FVector(0, 0, 1));
+		FPlane GroundPlane(FVector(0,0,Properties->GroundHeight) , FVector(0, 0, 1));
 		NewActorPos = FMath::RayPlaneIntersection(ClickPos.WorldRay.Origin, ClickPos.WorldRay.Direction, GroundPlane);
 	}
 
diff --git a/Engine/Plugins/Experimental/SampleToolsEditorMode/Source/Private/SampleTools/CreateActorSampleTool.h b/Engine/Plugins/Experimental/SampleToolsEditorMode/Source/Private/SampleTools/CreateActorSampleTool.h
--- a/Engine/Plugins/Experimental/SampleToolsEditorMode/Source/Private/SampleTools/CreateActorSampleTool.h
+++ b/Engine/Plugins/Experimental/SampleToolsEditorMode/Source/Private/SampleTools/CreateActorSampleTool.h
@@ -31,6 +31,29 @@ public:
 
 
 
+/**
+ * Settings UObject for UCreateActorSampleTool. This UClass inherits from UInteractiveToolPropertySet,
+ * which provides an OnModified delegate that the Tool will listen to for changes in property values.
+ */
+UCLASS(Transient)
+class SAMPLETOOLSEDITORMODE_API UCreateActorSampleToolProperties : public UInteractiveToolPropertySet
+{
+	GENERATED_BODY()
+public:
+	UCreateActorSampleToolProperties();
+
+	/** Place actors on existing objects */
+	UPROPERTY(EditAnywhere, Category = Options, meta = (DisplayName = "Place On Objects"))
+	bool PlaceOnObjects;
+
+	/** Height of ground plane */
+	UPROPERTY(EditAnywhere, Category = Options, meta = (DisplayName = "Ground Height", UIMin = "-1000.0", UIMax = "1000.0", ClampMin = "-1000000", ClampMax = "1000000.0"))
+	float GroundHeight;
+};
+
+
+
+
 /**
  * UCreateActorSampleTool is an example Tool that drops an empty Actor at each position the user 
  * clicks left mouse button. The Actors are placed at the first ray intersection in the scene,
@@ -53,13 +76,8 @@ public:
 
 
 protected:
-	/** Place actors on existing objects */
-	UPROPERTY(EditAnywhere, Category = Options, meta = (DisplayName = "Place On Objects"))
-	bool PlaceOnObjects;
-
-	/** Height of ground plane */
-	UPROPERTY(EditAnywhere, Category = Options, meta = (DisplayName = "Ground Height", UIMin = "-1000.0", UIMax = "1000.0", ClampMin = "-1000000", ClampMax = "1000000.0"))
-	float GroundHeight;
+	UPROPERTY()
+	UCreateActorSampleToolProperties* Properties;
 
 
 protected:
diff --git a/Engine/Plugins/Experimental/SampleToolsEditorMode/Source/Private/SampleTools/DrawCurveOnMeshSampleTool.cpp b/Engine/Plugins/Experimental/SampleToolsEditorMode/Source/Private/SampleTools/DrawCurveOnMeshSampleTool.cpp
--- a/Engine/Plugins/Experimental/SampleToolsEditorMode/Source/Private/SampleTools/DrawCurveOnMeshSampleTool.cpp
+++ b/Engine/Plugins/Experimental/SampleToolsEditorMode/Source/Private/SampleTools/DrawCurveOnMeshSampleTool.cpp
@@ -1,6 +1,7 @@
 // Copyright 1998-2019 Epic Games, Inc. All Rights Reserved.
 
 #include "DrawCurveOnMeshSampleTool.h"
+#include "InteractiveToolManager.h"
 #include "ToolBuilderUtil.h"
 #include "SceneManagement.h"   // FPrimitiveDrawInterface
 
@@ -12,9 +13,9 @@
  */
 
 
-UMeshSurfacePointTool* UDrawCurveOnMeshSampleToolBuilder::CreateNewTool(const FToolBuilderState & SceneState) const
+UMeshSurfacePointTool* UDrawCurveOnMeshSampleToolBuilder::CreateNewTool(const FToolBuilderState& SceneState) const
 {
-	UDrawCurveOnMeshSampleTool* NewTool = NewObject<UDrawCurveOnMeshSampleTool>();
+	UDrawCurveOnMeshSampleTool* NewTool = NewObject<UDrawCurveOnMeshSampleTool>(SceneState.ToolManager);
 	return NewTool;
 }
 
diff --git a/Engine/Plugins/Experimental/SampleToolsEditorMode/Source/Private/SampleTools/MeasureDistanceSampleTool.h b/Engine/Plugins/Experimental/SampleToolsEditorMode/Source/Private/SampleTools/MeasureDistanceSampleTool.h
new file mode 100644
--- /dev/null
+++ b/Engine/Plugins/Experimental/SampleToolsEditorMode/Source/Private/SampleTools/MeasureDistanceSampleTool.h
@@ -0,0 +1,102 @@
+// Copyright 1998-2019 Epic Games, Inc. All Rights Reserved.
+
+#pragma once
+
+#include "CoreMinimal.h"
+#include "InteractiveToolBuilder.h"
+#include "BaseTools/ClickDragTool.h"
+#include "MeasureDistanceSampleTool.generated.h"
+
+
+/**
+ * Builder for UMeasureDistanceSampleTool
+ */
+UCLASS()
+class SAMPLETOOLSEDITORMODE_API UMeasureDistanceSampleToolBuilder : public UInteractiveToolBuilder
+{
+	GENERATED_BODY()
+
+public:
+	virtual bool CanBuildTool(const FToolBuilderState& SceneState) const override;
+	virtual UInteractiveTool* BuildTool(const FToolBuilderState& SceneState) const override;
+};
+
+
+/**
+ * Property set for the UMeasureDistanceSampleTool
+ */
+UCLASS(Transient)
+class SAMPLETOOLSEDITORMODE_API UMeasureDistanceProperties : public UInteractiveToolPropertySet
+{
+	GENERATED_BODY()
+
+public:
+	UMeasureDistanceProperties();
+
+	/** First point of measurement */
+	UPROPERTY(EditAnywhere, Category = Options)
+	FVector StartPoint;
+
+	/** Second point of measurement */
+	UPROPERTY(EditAnywhere, Category = Options)
+	FVector EndPoint;
+	
+	/** Current distance measurement */
+	UPROPERTY(EditAnywhere, Category = Options)
+	float Distance;
+};
+
+
+
+/**
+ * UMeasureDistanceSampleTool is an example Tool that allows the user to measure the 
+ * distance between two points. The first point is set by click-dragging the mouse, and
+ * the second point is set by shift-click-dragging the mouse.
+ */
+UCLASS()
+class SAMPLETOOLSEDITORMODE_API UMeasureDistanceSampleTool : public UInteractiveTool, public IClickDragBehaviorTarget
+{
+	GENERATED_BODY()
+
+public:
+	UMeasureDistanceSampleTool();
+
+	virtual void SetWorld(UWorld* World);
+
+	// UInteractiveTool overrides
+
+	virtual void Setup() override;
+	virtual void Render(IToolsContextRenderAPI* RenderAPI) override;
+	virtual void OnPropertyModified(UObject* PropertySet, UProperty* Property) override;
+
+	// IClickDragBehaviorTarget implementation
+
+	virtual bool CanBeginClickDragSequence(const FInputDeviceRay& PressPos) override;
+	virtual void OnClickPress(const FInputDeviceRay& PressPos) override;
+	virtual void OnClickDrag(const FInputDeviceRay& DragPos) override;
+	// these are not used in this Tool
+	virtual void OnClickRelease(const FInputDeviceRay& ReleasePos) override {}
+	virtual void OnTerminateDragSequence() override {}
+
+	// IModifierToggleBehaviorTarget implementation (inherited via IClickDragBehaviorTarget)
+
+	virtual void OnUpdateModifierState(int ModifierID, bool bIsOn) override;
+
+
+protected:
+	/** Properties of the tool are stored here */
+	UPROPERTY()
+	UMeasureDistanceProperties* Properties;
+
+
+protected:
+	UWorld* TargetWorld = nullptr;		// target World we will raycast into
+
+	static const int MoveSecondPointModifierID = 1;		// identifier we associate with the shift key
+	bool bSecondPointModifierDown = false;				// flag we use to keep track of modifier state
+	bool bMoveSecondPoint = false;						// flag we use to keep track of which point we are moving during a press-drag
+
+	bool FindRayHit(const FRay& WorldRay, FVector& HitPos);		// raycasts into World
+	void UpdatePosition(const FRay& WorldRay);					// updates first or second point based on raycast
+	void UpdateDistance();										// updates distance
+};
\ No newline at end of file
diff --git a/Engine/Plugins/Experimental/SampleToolsEditorMode/Source/Private/SampleToolsEditorMode.cpp b/Engine/Plugins/Experimental/SampleToolsEditorMode/Source/Private/SampleToolsEditorMode.cpp
--- a/Engine/Plugins/Experimental/SampleToolsEditorMode/Source/Private/SampleToolsEditorMode.cpp
+++ b/Engine/Plugins/Experimental/SampleToolsEditorMode/Source/Private/SampleToolsEditorMode.cpp
@@ -16,6 +16,7 @@
 #include "BaseTools/MeshSurfacePointTool.h"
 #include "SampleTools/CreateActorSampleTool.h"
 #include "SampleTools/DrawCurveOnMeshSampleTool.h"
+#include "SampleTools/MeasureDistanceSampleTool.h"
 
 // step 2: register a ToolBuilder in FSampleToolsEditorMode::Enter()
 // step 3: add a button in FSampleToolsEditorModeToolkit::Init()
@@ -184,10 +185,12 @@ void FSampleToolsEditorMode::Enter()
 	auto DrawCurveOnMeshSampleToolBuilder = NewObject<UDrawCurveOnMeshSampleToolBuilder>();
 	ToolsContext->ToolManager->RegisterToolType(TEXT("DrawCurveOnMeshSampleTool"), DrawCurveOnMeshSampleToolBuilder);
 
+	auto MeasureDistanceSampleToolBuilder = NewObject<UMeasureDistanceSampleToolBuilder>();
+	ToolsContext->ToolManager->RegisterToolType(TEXT("MeasureDistanceSampleTool"), MeasureDistanceSampleToolBuilder);
+
 	auto SurfacePointToolBuilder = NewObject<UMeshSurfacePointToolBuilder>();
 	ToolsContext->ToolManager->RegisterToolType(TEXT("SurfacePointTool"), SurfacePointToolBuilder);
 
-
 	// active tool type is not relevant here, we just set to default
 	ToolsContext->ToolManager->SelectActiveToolType(EToolSide::Left, TEXT("SurfacePointTool"));
 }
diff --git a/Engine/Plugins/Experimental/SampleToolsEditorMode/Source/Private/SampleToolsEditorModeToolkit.cpp b/Engine/Plugins/Experimental/SampleToolsEditorMode/Source/Private/SampleToolsEditorModeToolkit.cpp
--- a/Engine/Plugins/Experimental/SampleToolsEditorMode/Source/Private/SampleToolsEditorModeToolkit.cpp
+++ b/Engine/Plugins/Experimental/SampleToolsEditorMode/Source/Private/SampleToolsEditorModeToolkit.cpp
@@ -74,6 +74,14 @@ void FSampleToolsEditorModeToolkit::Init(const TSharedPtr<IToolkitHost>& InitToo
 					.IsEnabled_Lambda([this]() { return this->CanStartTool(TEXT("CreateActorSampleTool")); })
 				]
 
+			+ SVerticalBox::Slot()
+				.HAlign(HAlign_Center)
+				.AutoHeight()
+				[
+					SNew(SButton).Text(LOCTEXT("MeasureDistanceSampleToolLabel", "Measure Distance"))
+					.OnClicked_Lambda([this]() { return this->StartTool(TEXT("MeasureDistanceSampleTool")); })
+					.IsEnabled_Lambda([this]() { return this->CanStartTool(TEXT("MeasureDistanceSampleTool")); })
+				]
 
 			+ SVerticalBox::Slot()
 				.HAlign(HAlign_Center)
@@ -84,7 +92,6 @@ void FSampleToolsEditorModeToolkit::Init(const TSharedPtr<IToolkitHost>& InitToo
 					.IsEnabled_Lambda([this]() { return this->CanStartTool(TEXT("DrawCurveOnMeshSampleTool")); })
 				]
 
-
 			+ SVerticalBox::Slot()
 				.HAlign(HAlign_Center)
 				.AutoHeight()
diff --git a/Engine/Source/Editor/EditorInteractiveToolsFramework/Private/EdModeInteractiveToolsContext.cpp b/Engine/Source/Editor/EditorInteractiveToolsFramework/Private/EdModeInteractiveToolsContext.cpp
--- a/Engine/Source/Editor/EditorInteractiveToolsFramework/Private/EdModeInteractiveToolsContext.cpp
+++ b/Engine/Source/Editor/EditorInteractiveToolsFramework/Private/EdModeInteractiveToolsContext.cpp
@@ -361,6 +361,9 @@ bool UEdModeInteractiveToolsContext::CapturedMouseMove(FEditorViewportClient* In
 	{
 		FInputDeviceState InputState = CurrentInputState;
 		InputState.InputDevice = EInputDevices::Mouse;
+		InputState.SetKeyStates(
+			InViewportClient->IsShiftPressed(), InViewportClient->IsAltPressed(),
+			InViewportClient->IsCtrlPressed(), InViewportClient->IsCmdPressed());
 		InputState.Mouse.Delta2D = CurrentInputState.Mouse.Position2D - OldPosition;
 		InputRouter->PostInputEvent(InputState);
 		return true;
diff --git a/Engine/Source/Runtime/InteractiveToolsFramework/Private/BaseBehaviors/ClickDragBehavior.cpp b/Engine/Source/Runtime/InteractiveToolsFramework/Private/BaseBehaviors/ClickDragBehavior.cpp
new file mode 100644
--- /dev/null
+++ b/Engine/Source/Runtime/InteractiveToolsFramework/Private/BaseBehaviors/ClickDragBehavior.cpp
@@ -0,0 +1,77 @@
+// Copyright 1998-2019 Epic Games, Inc. All Rights Reserved.
+
+#include "BaseBehaviors/ClickDragBehavior.h"
+
+
+
+UClickDragInputBehavior::UClickDragInputBehavior()
+{
+}
+
+
+void UClickDragInputBehavior::Initialize(IClickDragBehaviorTarget* TargetIn)
+{
+	check(TargetIn != nullptr);
+	this->Target = TargetIn;
+}
+
+
+FInputCaptureRequest UClickDragInputBehavior::WantsCapture(const FInputDeviceState& Input)
+{
+	if (IsPressed(Input) && (ModifierCheckFunc == nullptr || ModifierCheckFunc(Input)) )
+	{
+		if ( Target->CanBeginClickDragSequence(GetDeviceRay(Input)) )
+		{
+			return FInputCaptureRequest::Begin(this, EInputCaptureSide::Any);
+		}
+	}
+	return FInputCaptureRequest::Ignore();
+}
+
+
+FInputCaptureUpdate UClickDragInputBehavior::BeginCapture(const FInputDeviceState& Input, EInputCaptureSide Side)
+{
+	Modifiers.UpdateModifiers(Input, Target);
+	OnClickPress(Input, Side);
+	return FInputCaptureUpdate::Begin(this, EInputCaptureSide::Any);
+}
+
+
+FInputCaptureUpdate UClickDragInputBehavior::UpdateCapture(const FInputDeviceState& Input, const FInputCaptureData& Data)
+{
+	Modifiers.UpdateModifiers(Input, Target);
+
+	if (IsReleased(Input)) 
+	{
+		OnClickRelease(Input, Data);
+		return FInputCaptureUpdate::End();
+	}
+	else
+	{
+		OnClickDrag(Input, Data);
+		return FInputCaptureUpdate::Continue();
+	}
+}
+
+
+void UClickDragInputBehavior::ForceEndCapture(const FInputCaptureData& Data)
+{
+	Target->OnTerminateDragSequence();
+}
+
+
+
+void UClickDragInputBehavior::OnClickPress(const FInputDeviceState& Input, EInputCaptureSide Side)
+{
+	Target->OnClickPress(GetDeviceRay(Input));
+}
+
+void UClickDragInputBehavior::OnClickDrag(const FInputDeviceState& Input, const FInputCaptureData& Data)
+{
+	Target->OnClickDrag(GetDeviceRay(Input));
+}
+
+void UClickDragInputBehavior::OnClickRelease(const FInputDeviceState& Input, const FInputCaptureData& Data)
+{
+	Target->OnClickRelease(GetDeviceRay(Input));
+}
diff --git a/Engine/Source/Runtime/InteractiveToolsFramework/Private/BaseTools/ClickDragTool.cpp b/Engine/Source/Runtime/InteractiveToolsFramework/Private/BaseTools/ClickDragTool.cpp
new file mode 100644
--- /dev/null
+++ b/Engine/Source/Runtime/InteractiveToolsFramework/Private/BaseTools/ClickDragTool.cpp
@@ -0,0 +1,72 @@
+// Copyright 1998-2019 Epic Games, Inc. All Rights Reserved.
+
+#include "BaseTools/ClickDragTool.h"
+#include "InteractiveToolManager.h"
+#include "BaseBehaviors/ClickDragBehavior.h"
+
+
+/*
+ * ToolBuilder
+ */
+
+bool UClickDragToolBuilder::CanBuildTool(const FToolBuilderState& SceneState) const
+{
+	return true;
+}
+
+UInteractiveTool* UClickDragToolBuilder::BuildTool(const FToolBuilderState& SceneState) const
+{
+	UClickDragTool* NewTool = NewObject<UClickDragTool>(SceneState.ToolManager);
+	return NewTool;
+}
+
+
+
+/*
+ * Tool
+ */
+
+
+void UClickDragTool::Setup()
+{
+	UInteractiveTool::Setup();
+
+	// add default mouse input behavior
+	UClickDragInputBehavior* MouseBehavior = NewObject<UClickDragInputBehavior>();
+	MouseBehavior->Initialize(this);
+	AddInputBehavior(MouseBehavior);
+}
+
+
+bool UClickDragTool::CanBeginClickDragSequence(const FInputDeviceRay& ClickPos)
+{
+	return true;
+}
+
+
+void UClickDragTool::OnClickPress(const FInputDeviceRay& ClickPos)
+{
+	// print debug message
+	GetToolManager()->PostMessage( 
+		FString::Printf( TEXT("UClickDragTool::OnClickPress at (%f,%f)"), ClickPos.ScreenPosition.X, ClickPos.ScreenPosition.Y), 
+		EToolMessageLevel::Internal );
+}
+
+void UClickDragTool::OnClickDrag(const FInputDeviceRay& ClickPos)
+{
+}
+
+
+void UClickDragTool::OnClickRelease(const FInputDeviceRay& ClickPos)
+{
+	// print debug message
+	GetToolManager()->PostMessage(
+		FString::Printf(TEXT("UClickDragTool::OnClickRelease at (%f,%f)"), ClickPos.ScreenPosition.X, ClickPos.ScreenPosition.Y),
+		EToolMessageLevel::Internal);
+}
+
+
+void UClickDragTool::OnTerminateDragSequence()
+{
+
+}
\ No newline at end of file
diff --git a/Engine/Source/Runtime/InteractiveToolsFramework/Private/BaseTools/SingleClickTool.cpp b/Engine/Source/Runtime/InteractiveToolsFramework/Private/BaseTools/SingleClickTool.cpp
--- a/Engine/Source/Runtime/InteractiveToolsFramework/Private/BaseTools/SingleClickTool.cpp
+++ b/Engine/Source/Runtime/InteractiveToolsFramework/Private/BaseTools/SingleClickTool.cpp
@@ -15,7 +15,7 @@ bool USingleClickToolBuilder::CanBuildTool(const FToolBuilderState& SceneState)
 
 UInteractiveTool* USingleClickToolBuilder::BuildTool(const FToolBuilderState& SceneState) const
 {
-	USingleClickTool* NewTool = NewObject<USingleClickTool>();
+	USingleClickTool* NewTool = NewObject<USingleClickTool>(SceneState.ToolManager);
 	return NewTool;
 }
 
diff --git a/Engine/Source/Runtime/InteractiveToolsFramework/Private/InteractiveTool.cpp b/Engine/Source/Runtime/InteractiveToolsFramework/Private/InteractiveTool.cpp
--- a/Engine/Source/Runtime/InteractiveToolsFramework/Private/InteractiveTool.cpp
+++ b/Engine/Source/Runtime/InteractiveToolsFramework/Private/InteractiveTool.cpp
@@ -44,6 +44,7 @@ const UInputBehaviorSet* UInteractiveTool::GetInputBehaviors() const
 
 void UInteractiveTool::AddToolPropertySource(UObject* PropertyObject)
 {
+	check(ToolPropertyObjects.Contains(PropertyObject) == false);
 	ToolPropertyObjects.Add(PropertyObject);
 }
 
diff --git a/Engine/Source/Runtime/InteractiveToolsFramework/Public/BaseBehaviors/BehaviorTargetInterfaces.h b/Engine/Source/Runtime/InteractiveToolsFramework/Public/BaseBehaviors/BehaviorTargetInterfaces.h
--- a/Engine/Source/Runtime/InteractiveToolsFramework/Public/BaseBehaviors/BehaviorTargetInterfaces.h
+++ b/Engine/Source/Runtime/InteractiveToolsFramework/Public/BaseBehaviors/BehaviorTargetInterfaces.h
@@ -53,6 +53,52 @@ public:
 
 
 
+/**
+ * Functions required to apply standard "Click-Drag" state machines to a target object.
+ * See UClickDragBehavior for an example of this kind of state machine.
+ */
+class IClickDragBehaviorTarget : public IModifierToggleBehaviorTarget
+{
+public:
+	virtual ~IClickDragBehaviorTarget() {}
+
+	/**
+	 * Test if target can begin click-drag interaction at this point
+	 * @param PressPos device position/ray at click point
+	 * @return true if target wants to begin sequence
+	 */
+	virtual bool CanBeginClickDragSequence(const FInputDeviceRay& PressPos) = 0;
+
+
+	/**
+	 * Notify Target that click press ocurred
+	 * @param PressPos device position/ray at click point
+	 */
+	virtual void OnClickPress(const FInputDeviceRay& PressPos) = 0;
+
+	/**
+	 * Notify Target that input position has changed
+	 * @param DragPos device position/ray at click point
+	 */
+	virtual void OnClickDrag(const FInputDeviceRay& DragPos) = 0;
+
+	/**
+	 * Notify Target that click release occurred
+	 * @param ReleasePos device position/ray at click point
+	 */
+	virtual void OnClickRelease(const FInputDeviceRay& ReleasePos) = 0;
+
+	/**
+	 * Notify Target that click-drag sequence has been explicitly terminated (eg by escape key)
+	 */
+	virtual void OnTerminateDragSequence() = 0;
+};
+
+
+
+
+
+
 /**
  * Target interface used by InputBehaviors that want to implement a multi-click sequence
  * (eg such as drawing a polygon with multiple clicks)
diff --git a/Engine/Source/Runtime/InteractiveToolsFramework/Public/BaseBehaviors/ClickDragBehavior.h b/Engine/Source/Runtime/InteractiveToolsFramework/Public/BaseBehaviors/ClickDragBehavior.h
new file mode 100644
--- /dev/null
+++ b/Engine/Source/Runtime/InteractiveToolsFramework/Public/BaseBehaviors/ClickDragBehavior.h
@@ -0,0 +1,76 @@
+// Copyright 1998-2019 Epic Games, Inc. All Rights Reserved.
+
+#pragma once
+
+#include "CoreMinimal.h"
+#include "BaseBehaviors/AnyButtonInputBehavior.h"
+#include "BaseBehaviors/BehaviorTargetInterfaces.h"
+#include "BaseBehaviors/InputBehaviorModifierStates.h"
+#include "ClickDragBehavior.generated.h"
+
+
+/**
+ * UClickDragInputBehavior implements a standard "button-click-drag"-style input behavior.
+ * An IClickDragBehaviorTarget instance must be provided which is manipulated by this behavior.
+ * 
+ * The state machine works as follows:
+ *    1) on input-device-button-press, call Target::CanBeginClickDragSequence to determine if capture should begin
+ *    2) on input-device-move, call Target::OnClickDrag
+ *    3) on input-device-button-release, call Target::OnClickRelease
+ *    
+ * If a ForceEndCapture occurs we call Target::OnTerminateDragSequence   
+ */
+UCLASS()
+class INTERACTIVETOOLSFRAMEWORK_API UClickDragInputBehavior : public UAnyButtonInputBehavior
+{
+	GENERATED_BODY()
+
+public:
+	UClickDragInputBehavior();
+
+	/**
+	 * The modifier set for this behavior
+	 */
+	FInputBehaviorModifierStates Modifiers;
+
+	/**
+	 * Initialize this behavior with the given Target
+	 * @param Target implementor of hit-test and on-clicked functions
+	 */
+	virtual void Initialize(IClickDragBehaviorTarget* Target);
+
+
+	/**
+	 * WantsCapture() will only return capture request if this function returns true (or is null)
+	 */
+	TFunction<bool(const FInputDeviceState&)> ModifierCheckFunc = nullptr;
+
+
+	// UInputBehavior implementation
+
+	virtual FInputCaptureRequest WantsCapture(const FInputDeviceState& Input) override;
+	virtual FInputCaptureUpdate BeginCapture(const FInputDeviceState& Input, EInputCaptureSide Side) override;
+	virtual FInputCaptureUpdate UpdateCapture(const FInputDeviceState& Input, const FInputCaptureData& Data) override;
+	virtual void ForceEndCapture(const FInputCaptureData& Data) override;
+
+
+protected:
+	/** Click Target object */
+	IClickDragBehaviorTarget* Target;
+
+	/**
+	 * Internal function that forwards click evens to Target::OnClickPress, you can customize behavior here
+	 */
+	virtual void OnClickPress(const FInputDeviceState& Input, EInputCaptureSide Side);
+
+	/**
+	 * Internal function that forwards click evens to Target::OnClickDrag, you can customize behavior here
+	 */
+	virtual void OnClickDrag(const FInputDeviceState& Input, const FInputCaptureData& Data);
+
+	/**
+	 * Internal function that forwards click evens to Target::OnClickRelease, you can customize behavior here
+	 */
+	virtual void OnClickRelease(const FInputDeviceState& Input, const FInputCaptureData& Data);
+};
+
diff --git a/Engine/Source/Runtime/InteractiveToolsFramework/Public/BaseBehaviors/InputBehaviorModifierStates.h b/Engine/Source/Runtime/InteractiveToolsFramework/Public/BaseBehaviors/InputBehaviorModifierStates.h
--- a/Engine/Source/Runtime/InteractiveToolsFramework/Public/BaseBehaviors/InputBehaviorModifierStates.h
+++ b/Engine/Source/Runtime/InteractiveToolsFramework/Public/BaseBehaviors/InputBehaviorModifierStates.h
@@ -41,6 +41,15 @@ public:
 	}
 
 
+	/**
+	 * @return true if any modifiers have been registered
+	 */
+	bool HasModifiers() const
+	{
+		return (ModifierIDs.Num() > 0);
+	}
+
+
 	/**
 	 * Look up the current state of each registered modifier and pass to the target
 	 * @param Input the current input device state
diff --git a/Engine/Source/Runtime/InteractiveToolsFramework/Public/BaseTools/ClickDragTool.h b/Engine/Source/Runtime/InteractiveToolsFramework/Public/BaseTools/ClickDragTool.h
new file mode 100644
--- /dev/null
+++ b/Engine/Source/Runtime/InteractiveToolsFramework/Public/BaseTools/ClickDragTool.h
@@ -0,0 +1,81 @@
+// Copyright 1998-2019 Epic Games, Inc. All Rights Reserved.
+
+#pragma once
+
+#include "CoreMinimal.h"
+#include "InteractiveTool.h"
+#include "InteractiveToolBuilder.h"
+#include "BaseBehaviors/BehaviorTargetInterfaces.h"
+#include "ClickDragTool.generated.h"
+
+
+
+/**
+ * Builder for UClickDragTool
+ */
+UCLASS()
+class INTERACTIVETOOLSFRAMEWORK_API UClickDragToolBuilder : public UInteractiveToolBuilder
+{
+	GENERATED_BODY()
+
+public:
+	virtual bool CanBuildTool(const FToolBuilderState& SceneState) const override;
+	virtual UInteractiveTool* BuildTool(const FToolBuilderState& SceneState) const override;
+};
+
+
+/**
+ * UClickDragTool is a base tool that basically just implements IClickDragBehaviorTarget,
+ * and on setup registers a UClickDragInputBehavior. You can subclass this Tool to
+ * implement basic click-drag type Tools. If you want to do more advanced things, 
+ * like handle modifier buttons/keys, you will need to implement IClickDragBehaviorTarget yourself
+ */
+UCLASS()
+class INTERACTIVETOOLSFRAMEWORK_API UClickDragTool : public UInteractiveTool, public IClickDragBehaviorTarget
+{
+	GENERATED_BODY()
+
+public:
+
+	/**
+	 * Register default primary-button-click InputBehaviors
+	 */
+	virtual void Setup() override;
+
+
+	//
+	// IClickBehaviorTarget implementation
+	//
+
+	/**
+	 * Test if target can begin click-drag interaction at this point
+	 * @param ClickPos device position/ray at click point
+	 * @return true if target wants to begin sequence
+	 */
+	virtual bool CanBeginClickDragSequence(const FInputDeviceRay& PressPos) override;
+
+	/**
+	 * Notify Target that click press ocurred
+	 * @param ClickPos device position/ray at click point
+	 */
+	virtual void OnClickPress(const FInputDeviceRay& PressPos) override;
+
+	/**
+	 * Notify Target that input position has changed
+	 * @param ClickPos device position/ray at click point
+	 */
+	virtual void OnClickDrag(const FInputDeviceRay& DragPos) override;
+
+	/**
+	 * Notify Target that click release occurred
+	 * @param ClickPos device position/ray at click point
+	 */
+	virtual void OnClickRelease(const FInputDeviceRay& ReleasePos) override;
+
+	/**
+	 * Notify Target that click-drag sequence has been explicitly terminated (eg by escape key)
+	 */
+	virtual void OnTerminateDragSequence() override;
+
+};
+
