ToolsFramework: Add FInputRayHit that can be used by InputBehaviors and BehaviorTarget interfaces to return more detailed hit info (hit distance, integer ID, etc).
#rnx

diff --git a/Engine/Plugins/Experimental/SampleToolsEditorMode/Source/Private/SampleTools/MeasureDistanceSampleTool.cpp b/Engine/Plugins/Experimental/SampleToolsEditorMode/Source/Private/SampleTools/MeasureDistanceSampleTool.cpp
--- a/Engine/Plugins/Experimental/SampleToolsEditorMode/Source/Private/SampleTools/MeasureDistanceSampleTool.cpp
+++ b/Engine/Plugins/Experimental/SampleToolsEditorMode/Source/Private/SampleTools/MeasureDistanceSampleTool.cpp
@@ -92,11 +92,12 @@ void UMeasureDistanceSampleTool::OnUpdateModifierState(int ModifierID, bool bIsO
 
 
 
-bool UMeasureDistanceSampleTool::CanBeginClickDragSequence(const FInputDeviceRay& PressPos)
+FInputRayHit UMeasureDistanceSampleTool::CanBeginClickDragSequence(const FInputDeviceRay& PressPos)
 {
 	// we only start drag if press-down is on top of something we can raycast
 	FVector Temp;
-	return FindRayHit(PressPos.WorldRay, Temp);
+	FInputRayHit Result = FindRayHit(PressPos.WorldRay, Temp);
+	return Result;
 }
 
 void UMeasureDistanceSampleTool::OnClickPress(const FInputDeviceRay& PressPos)
@@ -113,7 +114,7 @@ void UMeasureDistanceSampleTool::OnClickDrag(const FInputDeviceRay& DragPos)
 
 
 
-bool UMeasureDistanceSampleTool::FindRayHit(const FRay& WorldRay, FVector& HitPos)
+FInputRayHit UMeasureDistanceSampleTool::FindRayHit(const FRay& WorldRay, FVector& HitPos)
 {
 	// trace a ray into the World
 	FCollisionObjectQueryParams QueryParams(FCollisionObjectQueryParams::AllObjects);
@@ -122,15 +123,16 @@ bool UMeasureDistanceSampleTool::FindRayHit(const FRay& WorldRay, FVector& HitPo
 	if (bHitWorld)
 	{
 		HitPos = Result.ImpactPoint;
-		return true;
+		return FInputRayHit(Result.Distance);
 	}
-	return false;
+	return FInputRayHit();
 }
 
 
 void UMeasureDistanceSampleTool::UpdatePosition(const FRay& WorldRay)
 {
-	if (FindRayHit(WorldRay, (bMoveSecondPoint) ? Properties->EndPoint : Properties->StartPoint))
+	FInputRayHit HitResult = FindRayHit(WorldRay, (bMoveSecondPoint) ? Properties->EndPoint : Properties->StartPoint);
+	if (HitResult.bHit)
 	{
 		UpdateDistance();
 	}
diff --git a/Engine/Plugins/Experimental/SampleToolsEditorMode/Source/Private/SampleTools/MeasureDistanceSampleTool.h b/Engine/Plugins/Experimental/SampleToolsEditorMode/Source/Private/SampleTools/MeasureDistanceSampleTool.h
--- a/Engine/Plugins/Experimental/SampleToolsEditorMode/Source/Private/SampleTools/MeasureDistanceSampleTool.h
+++ b/Engine/Plugins/Experimental/SampleToolsEditorMode/Source/Private/SampleTools/MeasureDistanceSampleTool.h
@@ -71,7 +71,7 @@ public:
 
 	// IClickDragBehaviorTarget implementation
 
-	virtual bool CanBeginClickDragSequence(const FInputDeviceRay& PressPos) override;
+	virtual FInputRayHit CanBeginClickDragSequence(const FInputDeviceRay& PressPos) override;
 	virtual void OnClickPress(const FInputDeviceRay& PressPos) override;
 	virtual void OnClickDrag(const FInputDeviceRay& DragPos) override;
 	// these are not used in this Tool
@@ -96,7 +96,7 @@ protected:
 	bool bSecondPointModifierDown = false;				// flag we use to keep track of modifier state
 	bool bMoveSecondPoint = false;						// flag we use to keep track of which point we are moving during a press-drag
 
-	bool FindRayHit(const FRay& WorldRay, FVector& HitPos);		// raycasts into World
+	FInputRayHit FindRayHit(const FRay& WorldRay, FVector& HitPos);		// raycasts into World
 	void UpdatePosition(const FRay& WorldRay);					// updates first or second point based on raycast
 	void UpdateDistance();										// updates distance
 };
\ No newline at end of file
diff --git a/Engine/Source/Runtime/Experimental/InteractiveToolsFramework/Private/BaseBehaviors/ClickDragBehavior.cpp b/Engine/Source/Runtime/Experimental/InteractiveToolsFramework/Private/BaseBehaviors/ClickDragBehavior.cpp
--- a/Engine/Source/Runtime/Experimental/InteractiveToolsFramework/Private/BaseBehaviors/ClickDragBehavior.cpp
+++ b/Engine/Source/Runtime/Experimental/InteractiveToolsFramework/Private/BaseBehaviors/ClickDragBehavior.cpp
@@ -20,9 +20,10 @@ FInputCaptureRequest UClickDragInputBehavior::WantsCapture(const FInputDeviceSta
 {
 	if (IsPressed(Input) && (ModifierCheckFunc == nullptr || ModifierCheckFunc(Input)) )
 	{
-		if ( Target->CanBeginClickDragSequence(GetDeviceRay(Input)) )
+		FInputRayHit HitResult = Target->CanBeginClickDragSequence(GetDeviceRay(Input));
+		if (HitResult.bHit)
 		{
-			return FInputCaptureRequest::Begin(this, EInputCaptureSide::Any);
+			return FInputCaptureRequest::Begin(this, EInputCaptureSide::Any, HitResult.HitDepth);
 		}
 	}
 	return FInputCaptureRequest::Ignore();
diff --git a/Engine/Source/Runtime/Experimental/InteractiveToolsFramework/Private/BaseTools/ClickDragTool.cpp b/Engine/Source/Runtime/Experimental/InteractiveToolsFramework/Private/BaseTools/ClickDragTool.cpp
--- a/Engine/Source/Runtime/Experimental/InteractiveToolsFramework/Private/BaseTools/ClickDragTool.cpp
+++ b/Engine/Source/Runtime/Experimental/InteractiveToolsFramework/Private/BaseTools/ClickDragTool.cpp
@@ -38,9 +38,9 @@ void UClickDragTool::Setup()
 }
 
 
-bool UClickDragTool::CanBeginClickDragSequence(const FInputDeviceRay& ClickPos)
+FInputRayHit UClickDragTool::CanBeginClickDragSequence(const FInputDeviceRay& ClickPos)
 {
-	return true;
+	return FInputRayHit(0.0f);
 }
 
 
diff --git a/Engine/Source/Runtime/Experimental/InteractiveToolsFramework/Public/BaseBehaviors/BehaviorTargetInterfaces.h b/Engine/Source/Runtime/Experimental/InteractiveToolsFramework/Public/BaseBehaviors/BehaviorTargetInterfaces.h
--- a/Engine/Source/Runtime/Experimental/InteractiveToolsFramework/Public/BaseBehaviors/BehaviorTargetInterfaces.h
+++ b/Engine/Source/Runtime/Experimental/InteractiveToolsFramework/Public/BaseBehaviors/BehaviorTargetInterfaces.h
@@ -65,9 +65,9 @@ public:
 	/**
 	 * Test if target can begin click-drag interaction at this point
 	 * @param PressPos device position/ray at click point
-	 * @return true if target wants to begin sequence
+	 * @return hit information at this point
 	 */
-	virtual bool CanBeginClickDragSequence(const FInputDeviceRay& PressPos) = 0;
+	virtual FInputRayHit CanBeginClickDragSequence(const FInputDeviceRay& PressPos) = 0;
 
 
 	/**
diff --git a/Engine/Source/Runtime/Experimental/InteractiveToolsFramework/Public/BaseTools/ClickDragTool.h b/Engine/Source/Runtime/Experimental/InteractiveToolsFramework/Public/BaseTools/ClickDragTool.h
--- a/Engine/Source/Runtime/Experimental/InteractiveToolsFramework/Public/BaseTools/ClickDragTool.h
+++ b/Engine/Source/Runtime/Experimental/InteractiveToolsFramework/Public/BaseTools/ClickDragTool.h
@@ -52,7 +52,7 @@ public:
 	 * @param ClickPos device position/ray at click point
 	 * @return true if target wants to begin sequence
 	 */
-	virtual bool CanBeginClickDragSequence(const FInputDeviceRay& PressPos) override;
+	virtual FInputRayHit CanBeginClickDragSequence(const FInputDeviceRay& PressPos) override;
 
 	/**
 	 * Notify Target that click press ocurred
diff --git a/Engine/Source/Runtime/Experimental/InteractiveToolsFramework/Public/InputState.h b/Engine/Source/Runtime/Experimental/InteractiveToolsFramework/Public/InputState.h
--- a/Engine/Source/Runtime/Experimental/InteractiveToolsFramework/Public/InputState.h
+++ b/Engine/Source/Runtime/Experimental/InteractiveToolsFramework/Public/InputState.h
@@ -6,6 +6,7 @@
 #include "UObject/ObjectMacros.h"
 #include "InputCoreTypes.h"
 #include "Math/UnrealMath.h"
+#include "InputState.generated.h"
 
 
 /**
@@ -63,6 +64,49 @@ struct FInputDeviceRay
 };
 
 
+
+/*
+ * FInputRayHit is returned by various hit-test interface functions.
+ * Generally this is intended to be returned as the result of a hit-test with a FInputDeviceRay 
+ */
+USTRUCT()
+struct INTERACTIVETOOLSFRAMEWORK_API FInputRayHit
+{
+	GENERATED_BODY()
+
+		/** true if ray hit something, false otherwise */
+		bool bHit;
+
+	/** distance along ray at which intersection occurred */
+	float HitDepth;
+
+	/** client-defined integer identifier for hit object/element/target/etc */
+	int32 HitIdentifier;
+
+	/** client-defined pointer for hit object/target */
+	void* HitOwner;
+
+	FInputRayHit()
+	{
+		bHit = false;
+		HitDepth = TNumericLimits<float>::Max();
+		HitIdentifier = 0;
+		HitOwner = nullptr;
+	}
+
+	explicit FInputRayHit(float HitDepthIn)
+	{
+		bHit = true;
+		HitDepth = HitDepthIn;
+		HitIdentifier = 0;
+		HitOwner = nullptr;
+	}
+};
+
+
+
+
+
 /**
  * Current State of a physical device button (mouse, key, etc) at a point in time.
  * Each "click" of a button should involve at minimum two separate state
