- Fixed some Linux CIS errors

diff --git a/Engine/Plugins/Experimental/SampleToolsEditorMode/Source/Private/SampleToolsEditorMode.cpp b/Engine/Plugins/Experimental/SampleToolsEditorMode/Source/Private/SampleToolsEditorMode.cpp
--- a/Engine/Plugins/Experimental/SampleToolsEditorMode/Source/Private/SampleToolsEditorMode.cpp
+++ b/Engine/Plugins/Experimental/SampleToolsEditorMode/Source/Private/SampleToolsEditorMode.cpp
@@ -4,7 +4,7 @@
 #include "SampleToolsEditorModeToolkit.h"
 #include "EdModeInteractiveToolsContext.h"
 #include "InteractiveToolManager.h"
-#include "SampleToolseditorModeCommands.h"
+#include "SampleToolsEditorModeCommands.h"
 
 
 //////////////////////////////////////////////////////////////////////////
diff --git a/Engine/Plugins/Experimental/SampleToolsEditorMode/Source/Public/SampleToolsEditorModeCommands.h b/Engine/Plugins/Experimental/SampleToolsEditorMode/Source/Public/SampleToolsEditorModeCommands.h
--- a/Engine/Plugins/Experimental/SampleToolsEditorMode/Source/Public/SampleToolsEditorModeCommands.h
+++ b/Engine/Plugins/Experimental/SampleToolsEditorMode/Source/Public/SampleToolsEditorModeCommands.h
@@ -1,5 +1,7 @@
 // Copyright Epic Games, Inc. All Rights Reserved.
 
+#pragma once
+
 #include "CoreMinimal.h"
 #include "Framework/Commands/Commands.h"
 
