Fixing CIS errors related to CL 14247116

diff --git a/Engine/Plugins/Experimental/GizmoEdMode/Source/GizmoEdMode/Private/GizmoEdMode.cpp b/Engine/Plugins/Experimental/GizmoEdMode/Source/GizmoEdMode/Private/GizmoEdMode.cpp
--- a/Engine/Plugins/Experimental/GizmoEdMode/Source/GizmoEdMode/Private/GizmoEdMode.cpp
+++ b/Engine/Plugins/Experimental/GizmoEdMode/Source/GizmoEdMode/Private/GizmoEdMode.cpp
@@ -23,7 +23,6 @@ UGizmoEdMode::UGizmoEdMode()
 		600
 	);
 	SettingsClass = UGizmoEdModeSettings::StaticClass();
-	ToolsContextClass = UEdModeInteractiveToolsContext::StaticClass();
 
 	AddFactory(NewObject<UDefaultAssetEditorGizmoFactory>());
 }
diff --git a/Engine/Plugins/Experimental/MotionTrailEditorMode/Source/MotionTrailEditorMode/Private/MotionTrailEditorMode.cpp b/Engine/Plugins/Experimental/MotionTrailEditorMode/Source/MotionTrailEditorMode/Private/MotionTrailEditorMode.cpp
--- a/Engine/Plugins/Experimental/MotionTrailEditorMode/Source/MotionTrailEditorMode/Private/MotionTrailEditorMode.cpp
+++ b/Engine/Plugins/Experimental/MotionTrailEditorMode/Source/MotionTrailEditorMode/Private/MotionTrailEditorMode.cpp
@@ -93,13 +93,8 @@ void UMotionTrailEditorMode::Exit()
 
 void UMotionTrailEditorMode::CreateToolkit()
 {
-	if (!Toolkit.IsValid())
-	{
-		FMotionTrailEditorModeToolkit* MotionTrailToolkit = new FMotionTrailEditorModeToolkit;
-		Toolkit = MakeShareable(MotionTrailToolkit);
-		Toolkit->Init(Owner->GetToolkitHost());
-	}
-	Super::CreateToolkit();
+	FMotionTrailEditorModeToolkit* MotionTrailToolkit = new FMotionTrailEditorModeToolkit;
+	Toolkit = MakeShareable(MotionTrailToolkit);
 }
 
 
diff --git a/Engine/Plugins/Experimental/MotionTrailEditorMode/Source/MotionTrailEditorMode/Private/MotionTrailEditorModeToolkit.cpp b/Engine/Plugins/Experimental/MotionTrailEditorMode/Source/MotionTrailEditorMode/Private/MotionTrailEditorModeToolkit.cpp
--- a/Engine/Plugins/Experimental/MotionTrailEditorMode/Source/MotionTrailEditorMode/Private/MotionTrailEditorModeToolkit.cpp
+++ b/Engine/Plugins/Experimental/MotionTrailEditorMode/Source/MotionTrailEditorMode/Private/MotionTrailEditorModeToolkit.cpp
@@ -33,9 +33,4 @@ FText FMotionTrailEditorModeToolkit::GetBaseToolkitName() const
 	return NSLOCTEXT("MotionTrailEditorModeToolkit", "DisplayName", "Motion Trail Editor Tool");
 }
 
-class UEdMode* FMotionTrailEditorModeToolkit::GetScriptableEditorMode() const
-{
-	return GLevelEditorModeTools().GetActiveScriptableMode("MotionTrailEditorMode");
-}
-
 #undef LOCTEXT_NAMESPACE
diff --git a/Engine/Plugins/Experimental/MotionTrailEditorMode/Source/MotionTrailEditorMode/Public/MotionTrailEditorModeToolkit.h b/Engine/Plugins/Experimental/MotionTrailEditorMode/Source/MotionTrailEditorMode/Public/MotionTrailEditorModeToolkit.h
--- a/Engine/Plugins/Experimental/MotionTrailEditorMode/Source/MotionTrailEditorMode/Public/MotionTrailEditorModeToolkit.h
+++ b/Engine/Plugins/Experimental/MotionTrailEditorMode/Source/MotionTrailEditorMode/Public/MotionTrailEditorModeToolkit.h
@@ -17,7 +17,6 @@ public:
 	/** IToolkit interface */
 	virtual FName GetToolkitFName() const override;
 	virtual FText GetBaseToolkitName() const override;
-	virtual class UEdMode* GetScriptableEditorMode() const override;
 
 private:
 
diff --git a/Engine/Plugins/Experimental/SampleToolsEditorMode/Source/Private/SampleToolsEditorMode.cpp b/Engine/Plugins/Experimental/SampleToolsEditorMode/Source/Private/SampleToolsEditorMode.cpp
--- a/Engine/Plugins/Experimental/SampleToolsEditorMode/Source/Private/SampleToolsEditorMode.cpp
+++ b/Engine/Plugins/Experimental/SampleToolsEditorMode/Source/Private/SampleToolsEditorMode.cpp
@@ -53,7 +53,7 @@ void USampleToolsEditorMode::Enter()
 	// The string name you pass to the ToolManager is used to select/activate your ToolBuilder later.
 	//////////////////////////////////////////////////////////////////////////
 	////////////////////////////////////////////////////////////////////////// 
-	FSampleToolsEditorModeCommands SampleToolCommands = FSampleToolsEditorModeCommands::Get();
+	const FSampleToolsEditorModeCommands& SampleToolCommands = FSampleToolsEditorModeCommands::Get();
 
 	auto CreateActorSampleToolBuilder = NewObject< UCreateActorSampleToolBuilder>(this);
 	CreateActorSampleToolBuilder->AssetAPI = ToolsContext->GetAssetAPI();
