Move SampleToolsEditorMode into Editor plugins folder

diff --git a/Engine/Plugins/Editor/SampleToolsEditorMode/SampleToolsEditorMode.uplugin b/Engine/Plugins/Editor/SampleToolsEditorMode/SampleToolsEditorMode.uplugin
new file mode 100644
--- /dev/null
+++ b/Engine/Plugins/Editor/SampleToolsEditorMode/SampleToolsEditorMode.uplugin
@@ -0,0 +1,24 @@
+{
+	"FileVersion": 3,
+	"Version": 1,
+	"VersionName": "0.1",
+	"FriendlyName": "Sample Tools Editor Mode",
+	"Description": "Sample Tools Mode includes a set of sample Tools demonstrating capabilities of the Interactive Tools Framework",
+	"Category": "Editor",
+	"CreatedBy": "Epic Games, Inc.",
+	"CreatedByURL": "https://epicgames.com",
+	"DocsURL": "",
+	"MarketplaceURL": "",
+	"SupportURL": "",
+	"CanContainContent": false,
+	"IsBetaVersion": true,
+	"Installed": false,
+	"Modules": [
+		{
+			"Name": "SampleToolsEditorMode",
+			"Type": "Editor",
+			"LoadingPhase": "Default"
+		}
+	],
+	"Plugins": []
+}
\ No newline at end of file
diff --git a/Engine/Plugins/Editor/SampleToolsEditorMode/Source/Private/SampleTools/CreateActorSampleTool.cpp b/Engine/Plugins/Editor/SampleToolsEditorMode/Source/Private/SampleTools/CreateActorSampleTool.cpp
new file mode 100644
--- /dev/null
+++ b/Engine/Plugins/Editor/SampleToolsEditorMode/Source/Private/SampleTools/CreateActorSampleTool.cpp
@@ -0,0 +1,101 @@
+// Copyright Epic Games, Inc. All Rights Reserved.
+
+#include "CreateActorSampleTool.h"
+#include "InteractiveToolManager.h"
+#include "Engine/World.h"
+#include "ToolContextInterfaces.h"
+
+// localization namespace
+#define LOCTEXT_NAMESPACE "UCreateActorSampleTool"
+
+/*
+ * ToolBuilder
+ */
+
+
+UInteractiveTool* UCreateActorSampleToolBuilder::BuildTool(const FToolBuilderState & SceneState) const
+{
+	UCreateActorSampleTool* NewTool = NewObject<UCreateActorSampleTool>(SceneState.ToolManager);
+	NewTool->SetWorld(SceneState.World);
+	return NewTool;
+}
+
+
+
+/*
+ * Tool
+ */
+
+
+UCreateActorSampleToolProperties::UCreateActorSampleToolProperties()
+{
+	PlaceOnObjects = true;
+	GroundHeight = 0.0f;
+}
+
+
+UCreateActorSampleTool::UCreateActorSampleTool()
+{
+}
+
+
+void UCreateActorSampleTool::SetWorld(UWorld* World)
+{
+	this->TargetWorld = World;
+}
+
+void UCreateActorSampleTool::Setup()
+{
+	USingleClickTool::Setup();
+
+	Properties = NewObject<UCreateActorSampleToolProperties>(this);
+	AddToolPropertySource(Properties);
+}
+
+
+
+
+void UCreateActorSampleTool::OnClicked(const FInputDeviceRay& ClickPos)
+{
+	// we will create actor at this position
+	FVector NewActorPos = FVector::ZeroVector;
+
+	// cast ray into world to find hit position
+	FVector RayStart = ClickPos.WorldRay.Origin;
+	FVector RayEnd = ClickPos.WorldRay.PointAt(999999);
+	FCollisionObjectQueryParams QueryParams(FCollisionObjectQueryParams::AllObjects);
+	FHitResult Result;
+	bool bHitWorld = TargetWorld->LineTraceSingleByObjectType(Result, RayStart, RayEnd, QueryParams);
+	if (Properties->PlaceOnObjects && bHitWorld)
+	{
+		NewActorPos = Result.ImpactPoint;
+	}
+	else
+	{
+		// hit the ground plane
+		FPlane GroundPlane(FVector(0,0,Properties->GroundHeight) , FVector(0, 0, 1));
+		NewActorPos = FMath::RayPlaneIntersection(ClickPos.WorldRay.Origin, ClickPos.WorldRay.Direction, GroundPlane);
+	}
+
+	// create new actor
+	FRotator Rotation(0.0f, 0.0f, 0.0f);
+	FActorSpawnParameters SpawnInfo;
+	AActor* NewActor = TargetWorld->SpawnActor<AActor>(FVector::ZeroVector, Rotation, SpawnInfo);
+
+	// create root component
+	USceneComponent* RootComponent = NewObject<USceneComponent>(NewActor, USceneComponent::GetDefaultSceneRootVariableName(), RF_Transactional);
+	RootComponent->Mobility = EComponentMobility::Movable;
+	RootComponent->bVisualizeComponent = true;
+	RootComponent->SetWorldTransform(FTransform(NewActorPos));
+
+	NewActor->SetRootComponent(RootComponent);
+	NewActor->AddInstanceComponent(RootComponent);
+	RootComponent->RegisterComponent();
+
+}
+
+
+
+
+
+#undef LOCTEXT_NAMESPACE
diff --git a/Engine/Plugins/Editor/SampleToolsEditorMode/Source/Private/SampleTools/CreateActorSampleTool.h b/Engine/Plugins/Editor/SampleToolsEditorMode/Source/Private/SampleTools/CreateActorSampleTool.h
new file mode 100644
--- /dev/null
+++ b/Engine/Plugins/Editor/SampleToolsEditorMode/Source/Private/SampleTools/CreateActorSampleTool.h
@@ -0,0 +1,75 @@
+// Copyright Epic Games, Inc. All Rights Reserved.
+
+#pragma once
+
+#include "BaseTools/SingleClickTool.h"
+#include "CreateActorSampleTool.generated.h"
+
+
+/**
+ * Builder for UCreateActorSampleTool
+ */
+UCLASS()
+class SAMPLETOOLSEDITORMODE_API UCreateActorSampleToolBuilder : public UInteractiveToolBuilder
+{
+	GENERATED_BODY()
+
+public:
+	virtual bool CanBuildTool(const FToolBuilderState& SceneState) const override { return true; }
+	virtual UInteractiveTool* BuildTool(const FToolBuilderState& SceneState) const override;
+};
+
+
+
+/**
+ * Settings UObject for UCreateActorSampleTool. This UClass inherits from UInteractiveToolPropertySet,
+ * which provides an OnModified delegate that the Tool will listen to for changes in property values.
+ */
+UCLASS(Transient)
+class SAMPLETOOLSEDITORMODE_API UCreateActorSampleToolProperties : public UInteractiveToolPropertySet
+{
+	GENERATED_BODY()
+public:
+	UCreateActorSampleToolProperties();
+
+	/** Place actors on existing objects */
+	UPROPERTY(EditAnywhere, Category = Options, meta = (DisplayName = "Place On Objects"))
+	bool PlaceOnObjects;
+
+	/** Height of ground plane */
+	UPROPERTY(EditAnywhere, Category = Options, meta = (DisplayName = "Ground Height", UIMin = "-1000.0", UIMax = "1000.0", ClampMin = "-1000000", ClampMax = "1000000.0"))
+	float GroundHeight;
+};
+
+
+
+
+/**
+ * UCreateActorSampleTool is an example Tool that drops an empty Actor at each position the user 
+ * clicks left mouse button. The Actors are placed at the first ray intersection in the scene,
+ * or on a ground plane if no scene objects are hit. All the action is in the ::OnClicked handler.
+ */
+UCLASS()
+class SAMPLETOOLSEDITORMODE_API UCreateActorSampleTool : public USingleClickTool
+{
+	GENERATED_BODY()
+
+public:
+	UCreateActorSampleTool();
+
+	virtual void SetWorld(UWorld* World);
+
+	virtual void Setup() override;
+
+	virtual void OnClicked(const FInputDeviceRay& ClickPos);
+
+
+protected:
+	UPROPERTY()
+	TObjectPtr<UCreateActorSampleToolProperties> Properties;
+
+
+protected:
+	/** target World we will raycast into and create Actor in */
+	UWorld* TargetWorld;
+};
diff --git a/Engine/Plugins/Editor/SampleToolsEditorMode/Source/Private/SampleTools/DrawCurveOnMeshSampleTool.cpp b/Engine/Plugins/Editor/SampleToolsEditorMode/Source/Private/SampleTools/DrawCurveOnMeshSampleTool.cpp
new file mode 100644
--- /dev/null
+++ b/Engine/Plugins/Editor/SampleToolsEditorMode/Source/Private/SampleTools/DrawCurveOnMeshSampleTool.cpp
@@ -0,0 +1,112 @@
+// Copyright Epic Games, Inc. All Rights Reserved.
+
+#include "DrawCurveOnMeshSampleTool.h"
+#include "Engine/HitResult.h"
+#include "InteractiveToolManager.h"
+#include "SceneManagement.h"   // FPrimitiveDrawInterface
+#include "ToolContextInterfaces.h"
+
+// localization namespace
+#define LOCTEXT_NAMESPACE "UDrawCurveOnMeshSampleTool"
+
+/*
+ * ToolBuilder
+ */
+
+
+UMeshSurfacePointTool* UDrawCurveOnMeshSampleToolBuilder::CreateNewTool(const FToolBuilderState& SceneState) const
+{
+	UDrawCurveOnMeshSampleTool* NewTool = NewObject<UDrawCurveOnMeshSampleTool>(SceneState.ToolManager);
+	return NewTool;
+}
+
+
+
+/*
+ * Tool
+ */
+
+UDrawCurveOnMeshSampleToolProperties::UDrawCurveOnMeshSampleToolProperties()
+{
+	Thickness = 4.0f;
+	DepthBias = 0.0f;
+	MinSpacing = 1.0;
+	NormalOffset = 0.25;
+	Color = FLinearColor(255, 0, 0);
+	bScreenSpace = true;
+}
+
+
+UDrawCurveOnMeshSampleTool::UDrawCurveOnMeshSampleTool()
+{
+}
+
+
+void UDrawCurveOnMeshSampleTool::Setup()
+{
+	UMeshSurfacePointTool::Setup();
+
+	// add settings object
+	Settings = NewObject<UDrawCurveOnMeshSampleToolProperties>(this, TEXT("Settings"));
+	AddToolPropertySource(Settings);
+}
+
+
+
+void UDrawCurveOnMeshSampleTool::Render(IToolsContextRenderAPI* RenderAPI)
+{
+	FPrimitiveDrawInterface* PDI = RenderAPI->GetPrimitiveDrawInterface();
+	
+	int NumPts = Positions.Num();
+	for (int i = 0; i < NumPts - 1; ++i)
+	{
+		FVector A = Positions[i] + Settings->NormalOffset*Normals[i];
+		FVector B = Positions[i+1] + Settings->NormalOffset*Normals[i+1];
+		PDI->DrawLine(A, B, Settings->Color, 0, Settings->Thickness, Settings->DepthBias, Settings->bScreenSpace);
+	}
+}
+
+
+
+
+void UDrawCurveOnMeshSampleTool::OnBeginDrag(const FRay& Ray)
+{
+	Positions.Reset();
+	Normals.Reset();
+	
+	FHitResult OutHit;
+	if (HitTest(Ray, OutHit))
+	{
+		Positions.Add(OutHit.ImpactPoint);
+		Normals.Add(OutHit.ImpactNormal);
+	}
+
+}
+
+void UDrawCurveOnMeshSampleTool::OnUpdateDrag(const FRay& Ray)
+{
+	FHitResult OutHit;
+	if (HitTest(Ray, OutHit))
+	{
+		if ( FVector::Dist(OutHit.ImpactPoint, Positions[Positions.Num()-1]) > Settings->MinSpacing)
+		{
+			Positions.Add(OutHit.ImpactPoint);
+			Normals.Add(OutHit.ImpactNormal);
+		}
+	}
+}
+
+void UDrawCurveOnMeshSampleTool::OnEndDrag(const FRay& Ray)
+{
+
+}
+
+void UDrawCurveOnMeshSampleTool::OnCancelDrag()
+{
+
+}
+
+
+
+
+#undef LOCTEXT_NAMESPACE
diff --git a/Engine/Plugins/Editor/SampleToolsEditorMode/Source/Private/SampleTools/DrawCurveOnMeshSampleTool.h b/Engine/Plugins/Editor/SampleToolsEditorMode/Source/Private/SampleTools/DrawCurveOnMeshSampleTool.h
new file mode 100644
--- /dev/null
+++ b/Engine/Plugins/Editor/SampleToolsEditorMode/Source/Private/SampleTools/DrawCurveOnMeshSampleTool.h
@@ -0,0 +1,89 @@
+// Copyright Epic Games, Inc. All Rights Reserved.
+
+#pragma once
+
+#include "BaseTools/MeshSurfacePointTool.h"
+#include "DrawCurveOnMeshSampleTool.generated.h"
+
+
+/**
+ * UMeshSurfacePointToolBuilder override for UDrawCurveOnMeshSampleTool
+ */
+UCLASS(Transient)
+class SAMPLETOOLSEDITORMODE_API UDrawCurveOnMeshSampleToolBuilder : public UMeshSurfacePointToolBuilder
+{
+	GENERATED_BODY()
+
+public:
+	virtual UMeshSurfacePointTool* CreateNewTool(const FToolBuilderState& SceneState) const override;
+};
+
+
+/**
+ * Settings UObject for UDrawCurveOnMeshSampleTool. This UClass inherits from UInteractiveToolPropertySet,
+ * which provides an OnModified delegate that the Tool will listen to for changes in property values.
+ */
+UCLASS(Transient)
+class SAMPLETOOLSEDITORMODE_API UDrawCurveOnMeshSampleToolProperties : public UInteractiveToolPropertySet
+{
+	GENERATED_BODY()
+public:
+	UDrawCurveOnMeshSampleToolProperties();
+
+	UPROPERTY(EditAnywhere, Category = Options)
+	FLinearColor Color;
+
+	UPROPERTY(EditAnywhere, Category = Options, meta = (DisplayName = "Thickness", UIMin = "0.25", UIMax = "10.0", ClampMin = "0.01", ClampMax = "1000.0"))
+	float Thickness;
+
+	UPROPERTY(EditAnywhere, Category = Options, meta = (DisplayName = "Min Spacing", UIMin = "0.01", UIMax = "10.0"))
+	float MinSpacing;
+
+	UPROPERTY(EditAnywhere, Category = Options, meta = (DisplayName = "Offset", UIMin = "0.0", UIMax = "10.0", ClampMin = "-1000.0", ClampMax = "1000.0"))
+	float NormalOffset;
+
+	UPROPERTY(EditAnywhere, Category = Options, meta = (DisplayName = "Depth Bias", UIMin = "-10.0", UIMax = "10.0"))
+	float DepthBias;
+
+	UPROPERTY(EditAnywhere, Category = Options)
+	bool bScreenSpace;
+};
+
+
+
+/**
+ * UDrawCurveOnMeshSampleTool is a sample Tool that allows the user to draw curves on the surface of
+ * a selected Mesh Component. The various rendering properties of the polycurve are exposed and can be tweaked.
+ * Nothing is done with the curve, it is just drawn by ::Render() and discarded when the Tool exits.
+ */
+UCLASS(Transient)
+class SAMPLETOOLSEDITORMODE_API UDrawCurveOnMeshSampleTool : public UMeshSurfacePointTool
+{
+	GENERATED_BODY()
+
+public:
+	UDrawCurveOnMeshSampleTool();
+
+	// UInteractiveTool API
+
+	virtual void Setup() override;
+
+	virtual void Render(IToolsContextRenderAPI* RenderAPI) override;
+
+	// UMeshSurfacePointTool API
+	virtual void OnBeginDrag(const FRay& Ray) override;
+	virtual void OnUpdateDrag(const FRay& Ray) override;
+	virtual void OnEndDrag(const FRay& Ray) override;
+	virtual void OnCancelDrag() override;
+
+protected:
+	UPROPERTY()
+	TObjectPtr<UDrawCurveOnMeshSampleToolProperties> Settings;
+
+	UPROPERTY()
+	TArray<FVector> Positions;
+
+	UPROPERTY()
+	TArray<FVector> Normals;
+
+};
diff --git a/Engine/Plugins/Editor/SampleToolsEditorMode/Source/Private/SampleTools/MeasureDistanceSampleTool.cpp b/Engine/Plugins/Editor/SampleToolsEditorMode/Source/Private/SampleTools/MeasureDistanceSampleTool.cpp
new file mode 100644
--- /dev/null
+++ b/Engine/Plugins/Editor/SampleToolsEditorMode/Source/Private/SampleTools/MeasureDistanceSampleTool.cpp
@@ -0,0 +1,163 @@
+// Copyright Epic Games, Inc. All Rights Reserved.
+
+#include "MeasureDistanceSampleTool.h"
+#include "InteractiveToolManager.h"
+#include "BaseBehaviors/ClickDragBehavior.h"
+
+// for raycast into World
+#include "Engine/World.h"
+
+#include "SceneManagement.h"
+#include "ToolContextInterfaces.h"
+
+// localization namespace
+#define LOCTEXT_NAMESPACE "UMeasureDistanceSampleTool"
+
+/*
+ * ToolBuilder
+ */
+
+
+bool UMeasureDistanceSampleToolBuilder::CanBuildTool(const FToolBuilderState & SceneState) const
+{
+	return true;
+}
+
+UInteractiveTool* UMeasureDistanceSampleToolBuilder::BuildTool(const FToolBuilderState & SceneState) const
+{
+	UMeasureDistanceSampleTool* NewTool = NewObject<UMeasureDistanceSampleTool>(SceneState.ToolManager);
+	NewTool->SetWorld(SceneState.World);
+	return NewTool;
+}
+
+
+
+/*
+ * Tool
+ */
+
+UMeasureDistanceProperties::UMeasureDistanceProperties()
+{
+	// initialize the points and distance to reasonable values
+	StartPoint = FVector(0,0,0);
+	EndPoint = FVector(0,0,100);
+	Distance = 100;
+}
+
+
+void UMeasureDistanceSampleTool::SetWorld(UWorld* World)
+{
+	check(World);
+	this->TargetWorld = World;
+}
+
+
+
+
+void UMeasureDistanceSampleTool::Setup()
+{
+	UInteractiveTool::Setup();
+
+	// Add default mouse input behavior
+	UClickDragInputBehavior* MouseBehavior = NewObject<UClickDragInputBehavior>();
+	// We will use the shift key to indicate that we should move the second point. 
+	// This call tells the Behavior to call our OnUpdateModifierState() function on mouse-down and mouse-move
+	MouseBehavior->Modifiers.RegisterModifier(MoveSecondPointModifierID, FInputDeviceState::IsShiftKeyDown);
+	MouseBehavior->Initialize(this);
+	AddInputBehavior(MouseBehavior);
+
+	// Create the property set and register it with the Tool
+	Properties = NewObject<UMeasureDistanceProperties>(this, "Measurement");
+	AddToolPropertySource(Properties);
+	
+	bSecondPointModifierDown = false;
+	bMoveSecondPoint = false;
+}
+
+
+void UMeasureDistanceSampleTool::OnUpdateModifierState(int ModifierID, bool bIsOn)
+{
+	// keep track of the "second point" modifier (shift key for mouse input)
+	if (ModifierID == MoveSecondPointModifierID)
+	{
+		bSecondPointModifierDown = bIsOn;
+	}
+}
+
+
+
+FInputRayHit UMeasureDistanceSampleTool::CanBeginClickDragSequence(const FInputDeviceRay& PressPos)
+{
+	// we only start drag if press-down is on top of something we can raycast
+	FVector Temp;
+	FInputRayHit Result = FindRayHit(PressPos.WorldRay, Temp);
+	return Result;
+}
+
+void UMeasureDistanceSampleTool::OnClickPress(const FInputDeviceRay& PressPos)
+{
+	// determine whether we are moving first or second point for the drag sequence
+	bMoveSecondPoint = bSecondPointModifierDown;
+	UpdatePosition(PressPos.WorldRay);
+}
+
+void UMeasureDistanceSampleTool::OnClickDrag(const FInputDeviceRay& DragPos)
+{
+	UpdatePosition(DragPos.WorldRay);
+}
+
+
+
+FInputRayHit UMeasureDistanceSampleTool::FindRayHit(const FRay& WorldRay, FVector& HitPos)
+{
+	// trace a ray into the World
+	FCollisionObjectQueryParams QueryParams(FCollisionObjectQueryParams::AllObjects);
+	FHitResult Result;
+	bool bHitWorld = TargetWorld->LineTraceSingleByObjectType(Result, WorldRay.Origin, WorldRay.PointAt(999999), QueryParams);
+	if (bHitWorld)
+	{
+		HitPos = Result.ImpactPoint;
+		return FInputRayHit(Result.Distance);
+	}
+	return FInputRayHit();
+}
+
+
+void UMeasureDistanceSampleTool::UpdatePosition(const FRay& WorldRay)
+{
+	FInputRayHit HitResult = FindRayHit(WorldRay, (bMoveSecondPoint) ? Properties->EndPoint : Properties->StartPoint);
+	if (HitResult.bHit)
+	{
+		UpdateDistance();
+	}
+}
+
+
+void UMeasureDistanceSampleTool::UpdateDistance()
+{
+	Properties->Distance = FVector::Distance(Properties->StartPoint, Properties->EndPoint);
+}
+
+
+void UMeasureDistanceSampleTool::OnPropertyModified(UObject* PropertySet, FProperty* Property)
+{
+	// if the user updated any of the property fields, update the distance
+	UpdateDistance();
+}
+
+
+
+void UMeasureDistanceSampleTool::Render(IToolsContextRenderAPI* RenderAPI)
+{
+	FPrimitiveDrawInterface* PDI = RenderAPI->GetPrimitiveDrawInterface();
+	// draw a thin line that shows through objects
+	PDI->DrawLine(Properties->StartPoint, Properties->EndPoint,
+		FColor(240, 16, 16), SDPG_Foreground, 2.0f, 0.0f, true);
+	// draw a thicker line that is depth-tested
+	PDI->DrawLine(Properties->StartPoint, Properties->EndPoint,
+		FColor(240, 16, 16), SDPG_World, 4.0f, 0.0f, true);
+}
+
+
+
+#undef LOCTEXT_NAMESPACE
diff --git a/Engine/Plugins/Editor/SampleToolsEditorMode/Source/Private/SampleTools/MeasureDistanceSampleTool.h b/Engine/Plugins/Editor/SampleToolsEditorMode/Source/Private/SampleTools/MeasureDistanceSampleTool.h
new file mode 100644
--- /dev/null
+++ b/Engine/Plugins/Editor/SampleToolsEditorMode/Source/Private/SampleTools/MeasureDistanceSampleTool.h
@@ -0,0 +1,100 @@
+// Copyright Epic Games, Inc. All Rights Reserved.
+
+#pragma once
+
+#include "BaseBehaviors/BehaviorTargetInterfaces.h"
+#include "InteractiveToolBuilder.h"
+#include "InteractiveTool.h"
+#include "MeasureDistanceSampleTool.generated.h"
+
+
+/**
+ * Builder for UMeasureDistanceSampleTool
+ */
+UCLASS()
+class SAMPLETOOLSEDITORMODE_API UMeasureDistanceSampleToolBuilder : public UInteractiveToolBuilder
+{
+	GENERATED_BODY()
+
+public:
+	virtual bool CanBuildTool(const FToolBuilderState& SceneState) const override;
+	virtual UInteractiveTool* BuildTool(const FToolBuilderState& SceneState) const override;
+};
+
+
+/**
+ * Property set for the UMeasureDistanceSampleTool
+ */
+UCLASS(Transient)
+class SAMPLETOOLSEDITORMODE_API UMeasureDistanceProperties : public UInteractiveToolPropertySet
+{
+	GENERATED_BODY()
+
+public:
+	UMeasureDistanceProperties();
+
+	/** First point of measurement */
+	UPROPERTY(EditAnywhere, Category = Options)
+	FVector StartPoint;
+
+	/** Second point of measurement */
+	UPROPERTY(EditAnywhere, Category = Options)
+	FVector EndPoint;
+	
+	/** Current distance measurement */
+	UPROPERTY(EditAnywhere, Category = Options)
+	float Distance;
+};
+
+
+
+/**
+ * UMeasureDistanceSampleTool is an example Tool that allows the user to measure the 
+ * distance between two points. The first point is set by click-dragging the mouse, and
+ * the second point is set by shift-click-dragging the mouse.
+ */
+UCLASS()
+class SAMPLETOOLSEDITORMODE_API UMeasureDistanceSampleTool : public UInteractiveTool, public IClickDragBehaviorTarget
+{
+	GENERATED_BODY()
+
+public:
+	virtual void SetWorld(UWorld* World);
+
+	// UInteractiveTool overrides
+
+	virtual void Setup() override;
+	virtual void Render(IToolsContextRenderAPI* RenderAPI) override;
+	virtual void OnPropertyModified(UObject* PropertySet, FProperty* Property) override;
+
+	// IClickDragBehaviorTarget implementation
+
+	virtual FInputRayHit CanBeginClickDragSequence(const FInputDeviceRay& PressPos) override;
+	virtual void OnClickPress(const FInputDeviceRay& PressPos) override;
+	virtual void OnClickDrag(const FInputDeviceRay& DragPos) override;
+	// these are not used in this Tool
+	virtual void OnClickRelease(const FInputDeviceRay& ReleasePos) override {}
+	virtual void OnTerminateDragSequence() override {}
+
+	// IModifierToggleBehaviorTarget implementation (inherited via IClickDragBehaviorTarget)
+
+	virtual void OnUpdateModifierState(int ModifierID, bool bIsOn) override;
+
+
+protected:
+	/** Properties of the tool are stored here */
+	UPROPERTY()
+	TObjectPtr<UMeasureDistanceProperties> Properties;
+
+
+protected:
+	UWorld* TargetWorld = nullptr;		// target World we will raycast into
+
+	static const int MoveSecondPointModifierID = 1;		// identifier we associate with the shift key
+	bool bSecondPointModifierDown = false;				// flag we use to keep track of modifier state
+	bool bMoveSecondPoint = false;						// flag we use to keep track of which point we are moving during a press-drag
+
+	FInputRayHit FindRayHit(const FRay& WorldRay, FVector& HitPos);		// raycasts into World
+	void UpdatePosition(const FRay& WorldRay);					// updates first or second point based on raycast
+	void UpdateDistance();										// updates distance
+};
diff --git a/Engine/Plugins/Editor/SampleToolsEditorMode/Source/Private/SampleToolsEditorMode.cpp b/Engine/Plugins/Editor/SampleToolsEditorMode/Source/Private/SampleToolsEditorMode.cpp
new file mode 100644
--- /dev/null
+++ b/Engine/Plugins/Editor/SampleToolsEditorMode/Source/Private/SampleToolsEditorMode.cpp
@@ -0,0 +1,78 @@
+// Copyright Epic Games, Inc. All Rights Reserved.
+
+#include "SampleToolsEditorMode.h"
+#include "SampleToolsEditorModeToolkit.h"
+#include "InteractiveToolManager.h"
+#include "SampleToolsEditorModeCommands.h"
+
+
+//////////////////////////////////////////////////////////////////////////
+////////////////////////////////////////////////////////////////////////// 
+// AddYourTool Step 1 - include the header file for your Tools here
+//////////////////////////////////////////////////////////////////////////
+////////////////////////////////////////////////////////////////////////// 
+#include "SampleTools/CreateActorSampleTool.h"
+#include "SampleTools/DrawCurveOnMeshSampleTool.h"
+#include "SampleTools/MeasureDistanceSampleTool.h"
+
+// step 2: register a ToolBuilder in FSampleToolsEditorMode::Enter()
+
+
+#define LOCTEXT_NAMESPACE "FSampleToolsEditorMode"
+
+const FEditorModeID USampleToolsEditorMode::EM_SampleToolsEditorModeId = TEXT("EM_SampleToolsEditorMode");
+
+
+USampleToolsEditorMode::USampleToolsEditorMode()
+{
+	Info = FEditorModeInfo(USampleToolsEditorMode::EM_SampleToolsEditorModeId,
+		LOCTEXT("SampleToolsEditorModeName", "SampleToolsEditorMode"),
+		FSlateIcon(),
+		true);
+}
+
+
+USampleToolsEditorMode::~USampleToolsEditorMode()
+{
+}
+
+
+void USampleToolsEditorMode::ActorSelectionChangeNotify()
+{
+	// @todo support selection change
+}
+
+void USampleToolsEditorMode::Enter()
+{
+	UEdMode::Enter();
+
+	//////////////////////////////////////////////////////////////////////////
+	//////////////////////////////////////////////////////////////////////////
+	// AddYourTool Step 2 - register the ToolBuilders for your Tools here.
+	// The string name you pass to the ToolManager is used to select/activate your ToolBuilder later.
+	//////////////////////////////////////////////////////////////////////////
+	////////////////////////////////////////////////////////////////////////// 
+	const FSampleToolsEditorModeCommands& SampleToolCommands = FSampleToolsEditorModeCommands::Get();
+
+	auto CreateActorSampleToolBuilder = NewObject< UCreateActorSampleToolBuilder>(this);
+	RegisterTool(SampleToolCommands.CreateActorTool, TEXT("CreateActorSampleTool"), CreateActorSampleToolBuilder);
+
+	RegisterTool(SampleToolCommands.DrawCurveOnMeshTool, TEXT("DrawCurveOnMeshSampleTool"), NewObject<UDrawCurveOnMeshSampleToolBuilder>(this));
+	RegisterTool(SampleToolCommands.MeasureDistanceTool, TEXT("MeasureDistanceSampleTool"), NewObject<UMeasureDistanceSampleToolBuilder>(this));
+	RegisterTool(SampleToolCommands.SurfacePointTool, TEXT("SurfacePointTool"), NewObject<UMeshSurfacePointToolBuilder>(this));
+
+	// active tool type is not relevant here, we just set to default
+	GetToolManager()->SelectActiveToolType(EToolSide::Left, TEXT("SurfacePointTool"));
+}
+
+void USampleToolsEditorMode::CreateToolkit()
+{
+	Toolkit = MakeShareable(new FSampleToolsEditorModeToolkit);
+}
+
+TMap<FName, TArray<TSharedPtr<FUICommandInfo>>> USampleToolsEditorMode::GetModeCommands() const
+{
+	return FSampleToolsEditorModeCommands::Get().GetCommands();
+}
+
+#undef LOCTEXT_NAMESPACE
diff --git a/Engine/Plugins/Editor/SampleToolsEditorMode/Source/Private/SampleToolsEditorModeCommands.cpp b/Engine/Plugins/Editor/SampleToolsEditorMode/Source/Private/SampleToolsEditorModeCommands.cpp
new file mode 100644
--- /dev/null
+++ b/Engine/Plugins/Editor/SampleToolsEditorMode/Source/Private/SampleToolsEditorModeCommands.cpp
@@ -0,0 +1,34 @@
+// Copyright Epic Games, Inc. All Rights Reserved.
+
+#include "SampleToolsEditorModeCommands.h"
+#include "Styling/AppStyle.h"
+
+#define LOCTEXT_NAMESPACE "FSampleToolsEditorMode"
+
+FSampleToolsEditorModeCommands::FSampleToolsEditorModeCommands()
+	: TCommands<FSampleToolsEditorModeCommands>("SampleToolsEditorMode",
+		NSLOCTEXT("SampleToolsEditorMode", "SampleToolsEditorModeCommands", "Sample Tools Editor Mode"),
+		NAME_None,
+		FAppStyle::GetAppStyleSetName())
+{
+}
+
+void FSampleToolsEditorModeCommands::RegisterCommands()
+{
+	TArray <TSharedPtr<FUICommandInfo>>& ToolCommands = Commands.FindOrAdd(NAME_Default);
+	UI_COMMAND(CreateActorTool, "Create Actor on Click", "Create Actor on Click", EUserInterfaceActionType::ToggleButton, FInputChord());
+	ToolCommands.Add(CreateActorTool);
+	UI_COMMAND(DrawCurveOnMeshTool, "Draw Curve On Mesh", "Draw Curve On Mesh", EUserInterfaceActionType::ToggleButton, FInputChord());
+	ToolCommands.Add(DrawCurveOnMeshTool);
+	UI_COMMAND(MeasureDistanceTool, "Measure Distance", "Measure Distance", EUserInterfaceActionType::ToggleButton, FInputChord());
+	ToolCommands.Add(MeasureDistanceTool);
+	UI_COMMAND(SurfacePointTool, "Surface Point Tool", "Surface Point Tool", EUserInterfaceActionType::ToggleButton, FInputChord());
+	ToolCommands.Add(SurfacePointTool);
+}
+
+TMap<FName, TArray<TSharedPtr<FUICommandInfo>>> FSampleToolsEditorModeCommands::GetCommands()
+{
+	return FSampleToolsEditorModeCommands::Get().Commands;
+}
+
+#undef LOCTEXT_NAMESPACE
diff --git a/Engine/Plugins/Editor/SampleToolsEditorMode/Source/Private/SampleToolsEditorModeModule.cpp b/Engine/Plugins/Editor/SampleToolsEditorMode/Source/Private/SampleToolsEditorModeModule.cpp
new file mode 100644
--- /dev/null
+++ b/Engine/Plugins/Editor/SampleToolsEditorMode/Source/Private/SampleToolsEditorModeModule.cpp
@@ -0,0 +1,26 @@
+// Copyright Epic Games, Inc. All Rights Reserved.
+
+#include "SampleToolsEditorModeModule.h"
+#include "Modules/ModuleManager.h"
+#include "SampleToolsEditorModeCommands.h"
+
+#define LOCTEXT_NAMESPACE "FSampleToolsEditorModeModule"
+
+void FSampleToolsEditorModeModule::StartupModule()
+{
+	// This code will execute after your module is loaded into memory; the exact timing is specified in the .uplugin file per-module
+
+	FSampleToolsEditorModeCommands::Register();
+}
+
+void FSampleToolsEditorModeModule::ShutdownModule()
+{
+	// This function may be called during shutdown to clean up your module.  For modules that support dynamic reloading,
+	// we call this function before unloading the module.
+
+	FSampleToolsEditorModeCommands::Unregister();
+}
+
+#undef LOCTEXT_NAMESPACE
+	
+IMPLEMENT_MODULE(FSampleToolsEditorModeModule, SampleToolsEditorMode)
diff --git a/Engine/Plugins/Editor/SampleToolsEditorMode/Source/Private/SampleToolsEditorModeToolkit.cpp b/Engine/Plugins/Editor/SampleToolsEditorMode/Source/Private/SampleToolsEditorModeToolkit.cpp
new file mode 100644
--- /dev/null
+++ b/Engine/Plugins/Editor/SampleToolsEditorMode/Source/Private/SampleToolsEditorModeToolkit.cpp
@@ -0,0 +1,37 @@
+// Copyright Epic Games, Inc. All Rights Reserved.
+
+#include "SampleToolsEditorModeToolkit.h"
+
+
+
+
+
+
+#define LOCTEXT_NAMESPACE "FSampleToolsEditorModeToolkit"
+
+FSampleToolsEditorModeToolkit::FSampleToolsEditorModeToolkit()
+{
+}
+
+void FSampleToolsEditorModeToolkit::Init(const TSharedPtr<IToolkitHost>& InitToolkitHost, TWeakObjectPtr<UEdMode> InOwningMode)
+{
+	FModeToolkit::Init(InitToolkitHost, InOwningMode);
+}
+
+void FSampleToolsEditorModeToolkit::GetToolPaletteNames(TArray<FName>& PaletteNames) const
+{
+	PaletteNames.Add(NAME_Default);
+}
+
+
+FName FSampleToolsEditorModeToolkit::GetToolkitFName() const
+{
+	return FName("SampleToolsEditorMode");
+}
+
+FText FSampleToolsEditorModeToolkit::GetBaseToolkitName() const
+{
+	return NSLOCTEXT("SampleToolsEditorModeToolkit", "DisplayName", "SampleToolsEditorMode Tool");
+}
+
+#undef LOCTEXT_NAMESPACE
diff --git a/Engine/Plugins/Editor/SampleToolsEditorMode/Source/Public/SampleToolsEditorMode.h b/Engine/Plugins/Editor/SampleToolsEditorMode/Source/Public/SampleToolsEditorMode.h
new file mode 100644
--- /dev/null
+++ b/Engine/Plugins/Editor/SampleToolsEditorMode/Source/Public/SampleToolsEditorMode.h
@@ -0,0 +1,41 @@
+// Copyright Epic Games, Inc. All Rights Reserved.
+
+#pragma once
+
+#include "CoreMinimal.h"
+#include "Tools/UEdMode.h"
+
+#include "SampleToolsEditorMode.generated.h"
+
+class IAssetGenerationAPI;
+
+/**
+ * This class provides an example of how to extend a UEdMode to add some simple tools
+ * using the InteractiveTools framework. The various UEdMode input event handlers (see UEdMode.h)
+ * forward events to a UEdModeInteractiveToolsContext instance, which 
+ * has all the logic for interacting with the InputRouter, ToolManager, etc.
+ * The functions provided here are the minimum to get started inserting some custom behavior.
+ * Take a look at the UEdMode markup for more extensibility options.
+ */
+UCLASS()
+class USampleToolsEditorMode : public UEdMode
+{
+	GENERATED_BODY()
+
+public:
+	const static FEditorModeID EM_SampleToolsEditorModeId;
+
+	USampleToolsEditorMode();
+	virtual ~USampleToolsEditorMode();
+
+	////////////////
+	// UEdMode interface
+	virtual void Enter() override;
+	virtual void ActorSelectionChangeNotify() override;
+	virtual void CreateToolkit() override;
+	virtual TMap<FName, TArray<TSharedPtr<FUICommandInfo>>> GetModeCommands() const override;
+
+	//////////////////
+	// End of UEdMode interface
+	//////////////////
+};
diff --git a/Engine/Plugins/Editor/SampleToolsEditorMode/Source/Public/SampleToolsEditorModeCommands.h b/Engine/Plugins/Editor/SampleToolsEditorMode/Source/Public/SampleToolsEditorModeCommands.h
new file mode 100644
--- /dev/null
+++ b/Engine/Plugins/Editor/SampleToolsEditorMode/Source/Public/SampleToolsEditorModeCommands.h
@@ -0,0 +1,23 @@
+// Copyright Epic Games, Inc. All Rights Reserved.
+
+#pragma once
+
+#include "CoreMinimal.h"
+#include "Framework/Commands/Commands.h"
+
+class FSampleToolsEditorModeCommands : public TCommands<FSampleToolsEditorModeCommands>
+{
+public:
+	FSampleToolsEditorModeCommands();
+
+	virtual void RegisterCommands() override;
+	static TMap<FName, TArray<TSharedPtr<FUICommandInfo>>> GetCommands();
+
+	TSharedPtr<FUICommandInfo> CreateActorTool;
+	TSharedPtr<FUICommandInfo> DrawCurveOnMeshTool;
+	TSharedPtr<FUICommandInfo> MeasureDistanceTool;
+	TSharedPtr<FUICommandInfo> SurfacePointTool;
+
+protected:
+	TMap<FName, TArray<TSharedPtr<FUICommandInfo>>> Commands;
+};
diff --git a/Engine/Plugins/Editor/SampleToolsEditorMode/Source/Public/SampleToolsEditorModeModule.h b/Engine/Plugins/Editor/SampleToolsEditorMode/Source/Public/SampleToolsEditorModeModule.h
new file mode 100644
--- /dev/null
+++ b/Engine/Plugins/Editor/SampleToolsEditorMode/Source/Public/SampleToolsEditorModeModule.h
@@ -0,0 +1,15 @@
+// Copyright Epic Games, Inc. All Rights Reserved.
+
+#pragma once
+
+#include "CoreMinimal.h"
+#include "Modules/ModuleManager.h"
+
+class FSampleToolsEditorModeModule : public IModuleInterface
+{
+public:
+
+	/** IModuleInterface implementation */
+	virtual void StartupModule() override;
+	virtual void ShutdownModule() override;
+};
diff --git a/Engine/Plugins/Editor/SampleToolsEditorMode/Source/Public/SampleToolsEditorModeToolkit.h b/Engine/Plugins/Editor/SampleToolsEditorMode/Source/Public/SampleToolsEditorModeToolkit.h
new file mode 100644
--- /dev/null
+++ b/Engine/Plugins/Editor/SampleToolsEditorMode/Source/Public/SampleToolsEditorModeToolkit.h
@@ -0,0 +1,28 @@
+// Copyright Epic Games, Inc. All Rights Reserved.
+
+#pragma once
+
+#include "CoreMinimal.h"
+#include "Toolkits/BaseToolkit.h"
+#include "InteractiveTool.h"
+#include "SampleToolsEditorMode.h"
+
+class IDetailsView;
+
+/**
+ * This FModeToolkit just creates a basic UI panel that allows various InteractiveTools to
+ * be initialized, and a DetailsView used to show properties of the active Tool.
+ */
+class FSampleToolsEditorModeToolkit : public FModeToolkit
+{
+public:
+	FSampleToolsEditorModeToolkit();
+	
+	// FModeToolkit interface 
+	virtual void Init(const TSharedPtr<IToolkitHost>& InitToolkitHost, TWeakObjectPtr<UEdMode> InOwningMode) override;
+	virtual void GetToolPaletteNames(TArray<FName>& PaletteNames) const override;
+
+	// IToolkit interface 
+	virtual FName GetToolkitFName() const override;
+	virtual FText GetBaseToolkitName() const override;
+};
diff --git a/Engine/Plugins/Editor/SampleToolsEditorMode/Source/SampleToolsEditorMode.Build.cs b/Engine/Plugins/Editor/SampleToolsEditorMode/Source/SampleToolsEditorMode.Build.cs
new file mode 100644
--- /dev/null
+++ b/Engine/Plugins/Editor/SampleToolsEditorMode/Source/SampleToolsEditorMode.Build.cs
@@ -0,0 +1,61 @@
+// Copyright Epic Games, Inc. All Rights Reserved.
+
+using UnrealBuildTool;
+
+public class SampleToolsEditorMode : ModuleRules
+{
+	public SampleToolsEditorMode(ReadOnlyTargetRules Target) : base(Target)
+	{
+		PCHUsage = ModuleRules.PCHUsageMode.UseExplicitOrSharedPCHs;
+		
+		PublicIncludePaths.AddRange(
+			new string[] {
+                //"ContentBrowser"
+				// ... add public include paths required here ...
+			}
+			);
+				
+		
+		PrivateIncludePaths.AddRange(
+			new string[] {
+				// ... add other private include paths required here ...
+			}
+			);
+			
+		
+		PublicDependencyModuleNames.AddRange(
+			new string[]
+			{
+				"Core"
+				// ... add other public dependencies that you statically link with here ...
+			}
+			);
+			
+		
+		PrivateDependencyModuleNames.AddRange(
+			new string[]
+			{
+				"CoreUObject",
+				"Engine",
+				"Slate",
+				"SlateCore",
+				"InputCore",
+				"EditorFramework",
+				"UnrealEd",
+                "ContentBrowser",
+                "LevelEditor",
+                "InteractiveToolsFramework",
+				"EditorInteractiveToolsFramework"
+				// ... add private dependencies that you statically link with here ...	
+			}
+            );
+		
+		
+		DynamicallyLoadedModuleNames.AddRange(
+			new string[]
+			{
+				// ... add any modules that your module loads dynamically here ...
+			}
+			);
+	}
+}
diff --git a/Engine/Plugins/Experimental/SampleToolsEditorMode/SampleToolsEditorMode.uplugin b/Engine/Plugins/Experimental/SampleToolsEditorMode/SampleToolsEditorMode.uplugin
deleted file mode 100644
diff --git a/Engine/Plugins/Experimental/SampleToolsEditorMode/Source/Private/SampleTools/CreateActorSampleTool.cpp b/Engine/Plugins/Experimental/SampleToolsEditorMode/Source/Private/SampleTools/CreateActorSampleTool.cpp
deleted file mode 100644
diff --git a/Engine/Plugins/Experimental/SampleToolsEditorMode/Source/Private/SampleTools/CreateActorSampleTool.h b/Engine/Plugins/Experimental/SampleToolsEditorMode/Source/Private/SampleTools/CreateActorSampleTool.h
deleted file mode 100644
diff --git a/Engine/Plugins/Experimental/SampleToolsEditorMode/Source/Private/SampleTools/DrawCurveOnMeshSampleTool.cpp b/Engine/Plugins/Experimental/SampleToolsEditorMode/Source/Private/SampleTools/DrawCurveOnMeshSampleTool.cpp
deleted file mode 100644
diff --git a/Engine/Plugins/Experimental/SampleToolsEditorMode/Source/Private/SampleTools/DrawCurveOnMeshSampleTool.h b/Engine/Plugins/Experimental/SampleToolsEditorMode/Source/Private/SampleTools/DrawCurveOnMeshSampleTool.h
deleted file mode 100644
diff --git a/Engine/Plugins/Experimental/SampleToolsEditorMode/Source/Private/SampleTools/MeasureDistanceSampleTool.cpp b/Engine/Plugins/Experimental/SampleToolsEditorMode/Source/Private/SampleTools/MeasureDistanceSampleTool.cpp
deleted file mode 100644
diff --git a/Engine/Plugins/Experimental/SampleToolsEditorMode/Source/Private/SampleTools/MeasureDistanceSampleTool.h b/Engine/Plugins/Experimental/SampleToolsEditorMode/Source/Private/SampleTools/MeasureDistanceSampleTool.h
deleted file mode 100644
diff --git a/Engine/Plugins/Experimental/SampleToolsEditorMode/Source/Private/SampleToolsEditorMode.cpp b/Engine/Plugins/Experimental/SampleToolsEditorMode/Source/Private/SampleToolsEditorMode.cpp
deleted file mode 100644
diff --git a/Engine/Plugins/Experimental/SampleToolsEditorMode/Source/Private/SampleToolsEditorModeCommands.cpp b/Engine/Plugins/Experimental/SampleToolsEditorMode/Source/Private/SampleToolsEditorModeCommands.cpp
deleted file mode 100644
diff --git a/Engine/Plugins/Experimental/SampleToolsEditorMode/Source/Private/SampleToolsEditorModeModule.cpp b/Engine/Plugins/Experimental/SampleToolsEditorMode/Source/Private/SampleToolsEditorModeModule.cpp
deleted file mode 100644
diff --git a/Engine/Plugins/Experimental/SampleToolsEditorMode/Source/Private/SampleToolsEditorModeToolkit.cpp b/Engine/Plugins/Experimental/SampleToolsEditorMode/Source/Private/SampleToolsEditorModeToolkit.cpp
deleted file mode 100644
diff --git a/Engine/Plugins/Experimental/SampleToolsEditorMode/Source/Public/SampleToolsEditorMode.h b/Engine/Plugins/Experimental/SampleToolsEditorMode/Source/Public/SampleToolsEditorMode.h
deleted file mode 100644
diff --git a/Engine/Plugins/Experimental/SampleToolsEditorMode/Source/Public/SampleToolsEditorModeCommands.h b/Engine/Plugins/Experimental/SampleToolsEditorMode/Source/Public/SampleToolsEditorModeCommands.h
deleted file mode 100644
diff --git a/Engine/Plugins/Experimental/SampleToolsEditorMode/Source/Public/SampleToolsEditorModeModule.h b/Engine/Plugins/Experimental/SampleToolsEditorMode/Source/Public/SampleToolsEditorModeModule.h
deleted file mode 100644
diff --git a/Engine/Plugins/Experimental/SampleToolsEditorMode/Source/Public/SampleToolsEditorModeToolkit.h b/Engine/Plugins/Experimental/SampleToolsEditorMode/Source/Public/SampleToolsEditorModeToolkit.h
deleted file mode 100644
diff --git a/Engine/Plugins/Experimental/SampleToolsEditorMode/Source/SampleToolsEditorMode.Build.cs b/Engine/Plugins/Experimental/SampleToolsEditorMode/Source/SampleToolsEditorMode.Build.cs
deleted file mode 100644
